UE-133356 HoloLens platform does not show in Platforms dropdown
UE-134396 Engine not taking correct HoloLens SDK
-Because the HoloLensTargetPlatform module was defined in a plugin it was missed by code which checks for the existence of such modules very early in editor startup, before any plugins are loaded.  This meant that parts of the editor never found that compiled support existed.  Other parts did so the packaging options would show up, but the menu layout was incorrect.  Possibly other features were also breaking based on this cached data saying that the platform had no support.
-We are moving the modules out of the HoloLensAR plugin into Engine/Source/Developer/HoloLens, like all other non-extension platforms.
-It was also discovered that HoloLensPlatformProperties did not enable virtual texture support.  Because HoloLens currently shares SM5 shaders with windows and the code that decides what platform properties to use for SM5 simply applies the first SM5 properties it finds shaders could be built with or without virtual texture support arbitrarily.  As a temporary fix we have enabled virtual texture support for hololens.  The properties are now the same and it doesn't matter which is chosen.  At present the windows properties are found first.
-Deleting the now empty HoloLensAR plugin.  This results in a pop up warning when opening a project which had this plugin enabled, but 5.0 seems as good a time as any for that to happen.
-Removed the now missing HoloLensAR from several projects and templates that enabled it.
-Fixed a project material that needed a new flag set for ue5.

diff --git a/Engine/Plugins/Runtime/AR/Microsoft/HoloLensAR/HoloLensAR.uplugin b/Engine/Plugins/Runtime/AR/Microsoft/HoloLensAR/HoloLensAR.uplugin
deleted file mode 100644
diff --git a/Engine/Plugins/Runtime/AR/Microsoft/HoloLensAR/Source/HoloLensPlatformEditor/HoloLensPlatformEditor.Build.cs b/Engine/Plugins/Runtime/AR/Microsoft/HoloLensAR/Source/HoloLensPlatformEditor/HoloLensPlatformEditor.Build.cs
deleted file mode 100644
diff --git a/Engine/Plugins/Runtime/AR/Microsoft/HoloLensAR/Source/HoloLensPlatformEditor/Private/HoloLensImageResourcesCustomization.cpp b/Engine/Plugins/Runtime/AR/Microsoft/HoloLensAR/Source/HoloLensPlatformEditor/Private/HoloLensImageResourcesCustomization.cpp
deleted file mode 100644
diff --git a/Engine/Plugins/Runtime/AR/Microsoft/HoloLensAR/Source/HoloLensPlatformEditor/Private/HoloLensImageResourcesCustomization.h b/Engine/Plugins/Runtime/AR/Microsoft/HoloLensAR/Source/HoloLensPlatformEditor/Private/HoloLensImageResourcesCustomization.h
deleted file mode 100644
diff --git a/Engine/Plugins/Runtime/AR/Microsoft/HoloLensAR/Source/HoloLensPlatformEditor/Private/HoloLensLocalizedResources.h b/Engine/Plugins/Runtime/AR/Microsoft/HoloLensAR/Source/HoloLensPlatformEditor/Private/HoloLensLocalizedResources.h
deleted file mode 100644
diff --git a/Engine/Plugins/Runtime/AR/Microsoft/HoloLensAR/Source/HoloLensPlatformEditor/Private/HoloLensLocalizedResourcesCustomization.cpp b/Engine/Plugins/Runtime/AR/Microsoft/HoloLensAR/Source/HoloLensPlatformEditor/Private/HoloLensLocalizedResourcesCustomization.cpp
deleted file mode 100644
diff --git a/Engine/Plugins/Runtime/AR/Microsoft/HoloLensAR/Source/HoloLensPlatformEditor/Private/HoloLensLocalizedResourcesCustomization.h b/Engine/Plugins/Runtime/AR/Microsoft/HoloLensAR/Source/HoloLensPlatformEditor/Private/HoloLensLocalizedResourcesCustomization.h
deleted file mode 100644
diff --git a/Engine/Plugins/Runtime/AR/Microsoft/HoloLensAR/Source/HoloLensPlatformEditor/Private/HoloLensPlatformEditorModule.cpp b/Engine/Plugins/Runtime/AR/Microsoft/HoloLensAR/Source/HoloLensPlatformEditor/Private/HoloLensPlatformEditorModule.cpp
deleted file mode 100644
diff --git a/Engine/Plugins/Runtime/AR/Microsoft/HoloLensAR/Source/HoloLensPlatformEditor/Private/HoloLensTargetSettings.cpp b/Engine/Plugins/Runtime/AR/Microsoft/HoloLensAR/Source/HoloLensPlatformEditor/Private/HoloLensTargetSettings.cpp
deleted file mode 100644
diff --git a/Engine/Plugins/Runtime/AR/Microsoft/HoloLensAR/Source/HoloLensPlatformEditor/Private/HoloLensTargetSettings.h b/Engine/Plugins/Runtime/AR/Microsoft/HoloLensAR/Source/HoloLensPlatformEditor/Private/HoloLensTargetSettings.h
deleted file mode 100644
diff --git a/Engine/Plugins/Runtime/AR/Microsoft/HoloLensAR/Source/HoloLensPlatformEditor/Private/HoloLensTargetSettingsCustomization.cpp b/Engine/Plugins/Runtime/AR/Microsoft/HoloLensAR/Source/HoloLensPlatformEditor/Private/HoloLensTargetSettingsCustomization.cpp
deleted file mode 100644
diff --git a/Engine/Plugins/Runtime/AR/Microsoft/HoloLensAR/Source/HoloLensPlatformEditor/Private/HoloLensTargetSettingsCustomization.h b/Engine/Plugins/Runtime/AR/Microsoft/HoloLensAR/Source/HoloLensPlatformEditor/Private/HoloLensTargetSettingsCustomization.h
deleted file mode 100644
diff --git a/Engine/Plugins/Runtime/AR/Microsoft/HoloLensAR/Source/HoloLensPlatformEditor/Public/HoloLensPlatformEditor.h b/Engine/Plugins/Runtime/AR/Microsoft/HoloLensAR/Source/HoloLensPlatformEditor/Public/HoloLensPlatformEditor.h
deleted file mode 100644
diff --git a/Engine/Plugins/Runtime/AR/Microsoft/HoloLensAR/Source/HoloLensTargetPlatform/HoloLensTargetPlatform.Build.cs b/Engine/Plugins/Runtime/AR/Microsoft/HoloLensAR/Source/HoloLensTargetPlatform/HoloLensTargetPlatform.Build.cs
deleted file mode 100644
diff --git a/Engine/Plugins/Runtime/AR/Microsoft/HoloLensAR/Source/HoloLensTargetPlatform/Private/HoloLensDeviceDetector.cpp b/Engine/Plugins/Runtime/AR/Microsoft/HoloLensAR/Source/HoloLensTargetPlatform/Private/HoloLensDeviceDetector.cpp
deleted file mode 100644
diff --git a/Engine/Plugins/Runtime/AR/Microsoft/HoloLensAR/Source/HoloLensTargetPlatform/Private/HoloLensTargetDevice.cpp b/Engine/Plugins/Runtime/AR/Microsoft/HoloLensAR/Source/HoloLensTargetPlatform/Private/HoloLensTargetDevice.cpp
deleted file mode 100644
diff --git a/Engine/Plugins/Runtime/AR/Microsoft/HoloLensAR/Source/HoloLensTargetPlatform/Private/HoloLensTargetDevice.h b/Engine/Plugins/Runtime/AR/Microsoft/HoloLensAR/Source/HoloLensTargetPlatform/Private/HoloLensTargetDevice.h
deleted file mode 100644
diff --git a/Engine/Plugins/Runtime/AR/Microsoft/HoloLensAR/Source/HoloLensTargetPlatform/Private/HoloLensTargetPlatform.cpp b/Engine/Plugins/Runtime/AR/Microsoft/HoloLensAR/Source/HoloLensTargetPlatform/Private/HoloLensTargetPlatform.cpp
deleted file mode 100644
diff --git a/Engine/Plugins/Runtime/AR/Microsoft/HoloLensAR/Source/HoloLensTargetPlatform/Private/HoloLensTargetPlatform.h b/Engine/Plugins/Runtime/AR/Microsoft/HoloLensAR/Source/HoloLensTargetPlatform/Private/HoloLensTargetPlatform.h
deleted file mode 100644
diff --git a/Engine/Plugins/Runtime/AR/Microsoft/HoloLensAR/Source/HoloLensTargetPlatform/Private/HoloLensTargetPlatformModule.cpp b/Engine/Plugins/Runtime/AR/Microsoft/HoloLensAR/Source/HoloLensTargetPlatform/Private/HoloLensTargetPlatformModule.cpp
deleted file mode 100644
diff --git a/Engine/Plugins/Runtime/AR/Microsoft/HoloLensAR/Source/HoloLensTargetPlatform/Private/IHoloLensDeviceDetector.h b/Engine/Plugins/Runtime/AR/Microsoft/HoloLensAR/Source/HoloLensTargetPlatform/Private/IHoloLensDeviceDetector.h
deleted file mode 100644
diff --git a/Engine/Plugins/Runtime/WindowsMixedReality/WindowsMixedReality.uplugin b/Engine/Plugins/Runtime/WindowsMixedReality/WindowsMixedReality.uplugin
--- a/Engine/Plugins/Runtime/WindowsMixedReality/WindowsMixedReality.uplugin
+++ b/Engine/Plugins/Runtime/WindowsMixedReality/WindowsMixedReality.uplugin
@@ -88,15 +88,6 @@
 			"PlatformAllowList": [
 				"Win64"
 			]
-		},
-		{
-			"Name": "HoloLensAR",
-			"Type": "Runtime",
-			"LoadingPhase": "PostConfigInit",
-			"PlatformAllowList": [
-				"HoloLens",
-				"Win64"
-			]
 		}
 	],
 	"Plugins": [
diff --git a/Engine/Source/Developer/HoloLens/HoloLensPlatformEditor/HoloLensPlatformEditor.Build.cs b/Engine/Source/Developer/HoloLens/HoloLensPlatformEditor/HoloLensPlatformEditor.Build.cs
new file mode 100644
--- /dev/null
+++ b/Engine/Source/Developer/HoloLens/HoloLensPlatformEditor/HoloLensPlatformEditor.Build.cs
@@ -0,0 +1,36 @@
+// Copyright Epic Games, Inc. All Rights Reserved.
+
+using UnrealBuildTool;
+
+public class HoloLensPlatformEditor : ModuleRules
+{
+	public HoloLensPlatformEditor(ReadOnlyTargetRules Target) : base(Target)
+	{
+		PrivateDependencyModuleNames.AddRange(
+			new string[] {
+				"Core",
+				"CoreUObject",
+				"Settings",
+				"DeveloperToolSettings",
+				"DesktopPlatform",
+				"Slate",
+				"SlateCore",
+				"EditorStyle",
+				"PropertyEditor",
+				"SharedSettingsWidgets",
+				"AppFramework",
+				"DesktopWidgets",
+				"EditorFramework",
+				"UnrealEd",
+				"SourceControl",
+				"WindowsTargetPlatform", // For ECompilerVersion
+				"EngineSettings",
+				"Projects",
+				"gltfToolkit",
+				"AudioSettingsEditor",
+			}
+		);
+
+		PublicSystemLibraries.Add("crypt32.lib");
+	}
+}
diff --git a/Engine/Source/Developer/HoloLens/HoloLensPlatformEditor/Private/HoloLensImageResourcesCustomization.cpp b/Engine/Source/Developer/HoloLens/HoloLensPlatformEditor/Private/HoloLensImageResourcesCustomization.cpp
new file mode 100644
--- /dev/null
+++ b/Engine/Source/Developer/HoloLens/HoloLensPlatformEditor/Private/HoloLensImageResourcesCustomization.cpp
@@ -0,0 +1,166 @@
+// Copyright Epic Games, Inc. All Rights Reserved.
+
+#include "HoloLensImageResourcesCustomization.h"
+
+#include "DetailLayoutBuilder.h"
+#include "IDetailChildrenBuilder.h"
+#include "DetailWidgetRow.h"
+#include "SExternalImageReference.h"
+#include "Widgets/Text/STextBlock.h"
+#include "IExternalImagePickerModule.h"
+#include "ISourceControlModule.h"
+#include "PropertyHandle.h"
+#include "EditorDirectories.h"
+#include "Interfaces/IPluginManager.h"
+#include "Settings/ProjectPackagingSettings.h"
+#include "Widgets/Input/SCheckBox.h"
+#include "IDesktopPlatform.h"
+#include "Framework/Application/SlateApplication.h"
+#include "DesktopPlatformModule.h"
+#include "HAL/FileManager.h"
+
+#define LOCTEXT_NAMESPACE "HoloLensImagesCustomization"
+
+namespace
+{
+	FString GetPickerPath()
+	{
+		return FEditorDirectories::Get().GetLastDirectory(ELastDirectory::GENERIC_OPEN);
+	}
+
+
+	bool HandlePostExternalIconCopy(const FString& InChosenImage)
+	{
+		FEditorDirectories::Get().SetLastDirectory(ELastDirectory::GENERIC_OPEN, FPaths::GetPath(InChosenImage));
+		return true;
+	}
+
+	void AddWidgetForResourceImage(IDetailChildrenBuilder& ChildrenBuilder, const FString& RootPath, const FString& Culture, const FString& ImageFileName, const FText& ImageCaption, const FVector2D& ImageDimensions)
+	{
+		const FString DefaultEngineImageSubPath = FString::Printf(TEXT("Build/HoloLens/DefaultImages/%s.png"), *ImageFileName);
+		const FString DefaultGameImageSubPath = FString::Printf(TEXT("Build/HoloLens/Resources/%s/%s.png"), *Culture, *ImageFileName);
+
+		const FString EngineImagePath = FPaths::EngineDir() / DefaultEngineImageSubPath;
+		const FString ProjectImagePath = RootPath / DefaultGameImageSubPath;
+
+		TArray<FString> ImageExtensions;
+		ImageExtensions.Add(TEXT("png"));
+
+		ChildrenBuilder.AddCustomRow(ImageCaption)
+		.NameContent()
+		[
+			SNew(STextBlock)
+			.Text(ImageCaption)
+			.Font(IDetailLayoutBuilder::GetDetailFont())
+		]
+		.ValueContent()
+		.MinDesiredWidth(ImageDimensions.X)
+		[
+			SNew(SHorizontalBox)
+			+ SHorizontalBox::Slot()
+			.FillWidth(1.0f)
+			.VAlign(VAlign_Center)
+			[
+				SNew(SExternalImageReference, EngineImagePath, ProjectImagePath)
+				.FileDescription(ImageCaption)
+				.MaxDisplaySize(ImageDimensions)
+				.OnGetPickerPath(FOnGetPickerPath::CreateStatic(&GetPickerPath))
+				.OnPostExternalImageCopy(FOnPostExternalImageCopy::CreateStatic(&HandlePostExternalIconCopy))
+				.DeleteTargetWhenDefaultChosen(true)
+				.FileExtensions(ImageExtensions)
+				.DeletePreviousTargetWhenExtensionChanges(true)
+			]
+		];
+	}
+
+	void PickGltfFor(const FString& ModelPath)
+	{
+		IDesktopPlatform* DesktopPlatform = FDesktopPlatformModule::Get();
+		if (!DesktopPlatform)
+		{
+			return;
+		}
+
+		FText Title = FText::FromString("glb");
+		FString TitleExtensions = TEXT("*.glb");
+
+		TArray<FString> OutFiles;
+		const FString Filter = FString::Printf(TEXT("%s files (%s)|%s"), *Title.ToString(), *TitleExtensions, *TitleExtensions);
+		const FString DefaultPath = FEditorDirectories::Get().GetLastDirectory(ELastDirectory::GENERIC_OPEN);
+
+		if (DesktopPlatform->OpenFileDialog(nullptr, FText::Format(LOCTEXT("GltfPickerDialogTitle", "Choose a {0} file"), Title).ToString(), DefaultPath, TEXT(""), Filter, EFileDialogFlags::None, OutFiles))
+		{
+			check(OutFiles.Num() == 1);
+
+			FString SourceImagePath = FPaths::ConvertRelativePathToFull(OutFiles[0]);
+			FEditorDirectories::Get().SetLastDirectory(ELastDirectory::GENERIC_OPEN, FPaths::GetPath(SourceImagePath));
+
+			if (ModelPath != SourceImagePath)
+			{
+				IFileManager::Get().Copy(*ModelPath, *SourceImagePath);
+			}
+		}
+
+	}
+
+	void AddWidgetForResourceGltf(IDetailChildrenBuilder& ChildrenBuilder, const FString& RootPath, const FString& Culture, const FString& ModelFileName, const FText& ModelCaption)
+	{
+		const FString ModelPath = RootPath / TEXT("Build") / TEXT("HoloLens") / TEXT("Resources") / Culture / (ModelFileName + TEXT(".glb"));
+
+		ChildrenBuilder.AddCustomRow(ModelCaption)
+		.NameContent()
+		[
+			SNew(STextBlock)
+			.Text(ModelCaption)
+			.Font(IDetailLayoutBuilder::GetDetailFont())
+		]
+		.ValueContent()
+		[
+			SNew(SCheckBox)
+			.IsChecked_Lambda([=]() -> ECheckBoxState 
+			{
+				return IFileManager::Get().FileExists(*ModelPath) ? ECheckBoxState::Checked : ECheckBoxState::Unchecked;
+			})
+			.OnCheckStateChanged_Lambda([=](ECheckBoxState NewAutoCloseState)
+			{
+				if (NewAutoCloseState == ECheckBoxState::Checked)
+				{
+					PickGltfFor(ModelPath);
+				}
+				else if (NewAutoCloseState == ECheckBoxState::Unchecked)
+				{
+					IFileManager::Get().Delete(*ModelPath, false, true, true);
+				}
+			})
+		];
+	}
+}
+
+TSharedRef<IPropertyTypeCustomization> FHoloLensCorePackageImagesCustomization::MakeInstance()
+{
+	return MakeShared<FHoloLensCorePackageImagesCustomization>();
+}
+
+void FHoloLensCorePackageImagesCustomization::CustomizeHeader(TSharedRef<class IPropertyHandle> StructPropertyHandle, class FDetailWidgetRow& HeaderRow, IPropertyTypeCustomizationUtils& StructCustomizationUtils)
+{
+	HeaderRow.NameContent()
+	[
+		StructPropertyHandle->CreatePropertyNameWidget()
+	];
+}
+
+void FHoloLensCorePackageImagesCustomization::CustomizeChildren(TSharedRef<IPropertyHandle> InStructPropertyHandle, IDetailChildrenBuilder& ChildBuilder, IPropertyTypeCustomizationUtils& StructCustomizationUtils)
+{
+	TSharedRef<IPropertyHandle> CultureIdProperty = InStructPropertyHandle->GetParentHandle()->GetChildHandle(FName("CultureId")).ToSharedRef();
+	FString CultureId;
+	CultureIdProperty->GetValue(CultureId);
+
+	AddWidgetForResourceImage(ChildBuilder, FPaths::ProjectDir(), CultureId, TEXT("Logo"), LOCTEXT("Square150x150Logo", "Square 150x150 Logo"), FVector2D(150.0f, 150.0f));
+	AddWidgetForResourceImage(ChildBuilder, FPaths::ProjectDir(), CultureId, TEXT("SmallLogo"), LOCTEXT("Square44x44Logo", "Square 44x44 Logo"), FVector2D(44.0f, 44.0f));
+	AddWidgetForResourceImage(ChildBuilder, FPaths::ProjectDir(), CultureId, TEXT("WideLogo"), LOCTEXT("Wide310x150Logo", "Wide 310x150 Logo"), FVector2D(310.0f, 150.0f));
+	AddWidgetForResourceImage(ChildBuilder, FPaths::ProjectDir(), CultureId, TEXT("SplashScreen"), LOCTEXT("SplashScreen", "Splash Screen"), FVector2D(620.0f, 300.0f));
+	AddWidgetForResourceImage(ChildBuilder, FPaths::ProjectDir(), CultureId, TEXT("StoreLogo"), LOCTEXT("StoreLogo", "Store Logo"), FVector2D(50.0f, 50.0f));
+	AddWidgetForResourceGltf(ChildBuilder, FPaths::ProjectDir(), CultureId, TEXT("3DLogo"), LOCTEXT("3DLogo", "3D Logo"));
+}
+
+#undef LOCTEXT_NAMESPACE
\ No newline at end of file
diff --git a/Engine/Source/Developer/HoloLens/HoloLensPlatformEditor/Private/HoloLensImageResourcesCustomization.h b/Engine/Source/Developer/HoloLens/HoloLensPlatformEditor/Private/HoloLensImageResourcesCustomization.h
new file mode 100644
--- /dev/null
+++ b/Engine/Source/Developer/HoloLens/HoloLensPlatformEditor/Private/HoloLensImageResourcesCustomization.h
@@ -0,0 +1,18 @@
+// Copyright Epic Games, Inc. All Rights Reserved.
+
+#pragma once
+
+#include "PropertyEditorModule.h"
+#include "IPropertyTypeCustomization.h"
+
+class FHoloLensCorePackageImagesCustomization : public IPropertyTypeCustomization
+{
+public:
+	static TSharedRef<IPropertyTypeCustomization> MakeInstance();
+
+	// IPropertyTypeCustomization interface
+	virtual void CustomizeHeader(TSharedRef<class IPropertyHandle> StructPropertyHandle, class FDetailWidgetRow& HeaderRow, IPropertyTypeCustomizationUtils& StructCustomizationUtils) override;
+
+	virtual void CustomizeChildren(TSharedRef<IPropertyHandle> InStructPropertyHandle, class IDetailChildrenBuilder& ChildBuilder, IPropertyTypeCustomizationUtils& StructCustomizationUtils) override;
+	// End of IPropertyTypeCustomization interface
+};
diff --git a/Engine/Source/Developer/HoloLens/HoloLensPlatformEditor/Private/HoloLensLocalizedResources.h b/Engine/Source/Developer/HoloLens/HoloLensPlatformEditor/Private/HoloLensLocalizedResources.h
new file mode 100644
--- /dev/null
+++ b/Engine/Source/Developer/HoloLens/HoloLensPlatformEditor/Private/HoloLensLocalizedResources.h
@@ -0,0 +1,47 @@
+// Copyright Epic Games, Inc. All Rights Reserved.
+
+#pragma once
+
+#include "HoloLensLocalizedResources.generated.h"
+
+USTRUCT()
+struct FHoloLensCorePackageStringResources
+{
+	GENERATED_BODY()
+
+	UPROPERTY(EditAnywhere, config, Category = Packaging)
+	FString PackageDisplayName;
+
+	UPROPERTY(EditAnywhere, config, Category = Packaging)
+	FString PublisherDisplayName;
+
+	UPROPERTY(EditAnywhere, config, Category = Packaging)
+	FString PackageDescription;
+
+	UPROPERTY(EditAnywhere, config, Category = Packaging)
+	FString ApplicationDisplayName;
+
+	UPROPERTY(EditAnywhere, config, Category = Packaging)
+	FString ApplicationDescription;
+};
+
+USTRUCT()
+struct FHoloLensCorePackageImageResources
+{
+	GENERATED_BODY()
+};
+
+USTRUCT()
+struct FHoloLensCorePackageLocalizedResources
+{
+	GENERATED_BODY()
+
+	UPROPERTY(EditAnywhere, config, Category = Packaging)
+	FString CultureId;
+
+	UPROPERTY(EditAnywhere, config, Category = Packaging, Meta=(ShowOnlyInnerProperties))
+	FHoloLensCorePackageStringResources Strings;
+
+	UPROPERTY(EditAnywhere, config, Category = Packaging)
+	FHoloLensCorePackageImageResources Images;
+};
diff --git a/Engine/Source/Developer/HoloLens/HoloLensPlatformEditor/Private/HoloLensLocalizedResourcesCustomization.cpp b/Engine/Source/Developer/HoloLens/HoloLensPlatformEditor/Private/HoloLensLocalizedResourcesCustomization.cpp
new file mode 100644
--- /dev/null
+++ b/Engine/Source/Developer/HoloLens/HoloLensPlatformEditor/Private/HoloLensLocalizedResourcesCustomization.cpp
@@ -0,0 +1,276 @@
+// Copyright Epic Games, Inc. All Rights Reserved.
+
+#include "HoloLensLocalizedResourcesCustomization.h"
+
+#include "DetailLayoutBuilder.h"
+#include "DetailCategoryBuilder.h"
+#include "DetailWidgetRow.h"
+#include "PropertyHandle.h"
+#include "PropertyCustomizationHelpers.h"
+#include "IDetailChildrenBuilder.h"
+#include "IDetailGroup.h"
+#include "Widgets/Input/STextComboBox.h"
+#include "Settings/ProjectPackagingSettings.h"
+
+#define LOCTEXT_NAMESPACE "HoloLensLocalizedResourcesCustomization"
+
+FHoloLensLocalizedResourcesNodeBuilder::FHoloLensLocalizedResourcesNodeBuilder(TSharedRef<IPropertyHandle> InLocalizedResourceCollectionProperty, const FString& InPluginName)
+	: OptionalPluginName(InPluginName)
+	, LocalizedResourceCollectionArray(InLocalizedResourceCollectionProperty->AsArray().ToSharedRef())
+{
+	// Ensure that we have a default set of fallback resources (empty culture id)
+	uint32 NumElements;
+	LocalizedResourceCollectionArray->GetNumElements(NumElements);
+
+	TSharedPtr<IPropertyHandle> DefaultResources;
+	for (uint32 i = 0; i < NumElements; ++i)
+	{
+		TSharedRef<IPropertyHandle> Element = LocalizedResourceCollectionArray->GetElement(i);
+		FString AppliesToDlcName;
+		TSharedPtr<IPropertyHandle> DlcPluginProperty = Element->GetChildHandle("AppliesToDlcPlugin");
+		if (DlcPluginProperty.IsValid())
+		{
+			DlcPluginProperty->GetValue(AppliesToDlcName);
+		}
+		FString CultureId;
+		Element->GetChildHandle("CultureId")->GetValue(CultureId);
+
+		if (AppliesToDlcName == OptionalPluginName && CultureId.IsEmpty())
+		{
+			DefaultResources = Element;
+		}
+	}
+
+	if (!DefaultResources.IsValid())
+	{
+		AddLocalizedEntry(TEXT(""));
+	}
+
+	InLocalizedResourceCollectionProperty->MarkHiddenByCustomization();
+	FSimpleDelegate ChangedDelegate = FSimpleDelegate::CreateRaw(this, &FHoloLensLocalizedResourcesNodeBuilder::OnNumElementsChanged);
+	LocalizedResourceCollectionArray->SetOnNumElementsChanged(ChangedDelegate);
+}
+
+FHoloLensLocalizedResourcesNodeBuilder::~FHoloLensLocalizedResourcesNodeBuilder()
+{
+	FSimpleDelegate Empty;
+	LocalizedResourceCollectionArray->SetOnNumElementsChanged(Empty);
+}
+
+void FHoloLensLocalizedResourcesNodeBuilder::AddLocalizedEntryForSelectedCulture()
+{
+	if (SelectedCulture.IsValid())
+	{
+		AddLocalizedEntry(*SelectedCulture->GetSelectedItem());
+	}
+}
+
+void FHoloLensLocalizedResourcesNodeBuilder::AddLocalizedEntry(const FString& CultureId)
+{
+	// Note: NumElements here wil
+	uint32 NumElements;
+	LocalizedResourceCollectionArray->GetNumElements(NumElements);
+
+	FPropertyAccess::Result AddResult = LocalizedResourceCollectionArray->AddItem();
+	check(AddResult == FPropertyAccess::Success);
+
+	// This little bit of magic forces a refresh of the property handle tree, ensuring
+	// that we can safely access element 0 (which we just moved our new entry into)
+	LocalizedResourceCollectionArray->MoveElementTo(NumElements, 0);
+
+	TSharedRef<IPropertyHandle> NewlyAddedItem = LocalizedResourceCollectionArray->GetElement(0);
+	NewlyAddedItem->NotifyPreChange();
+	if (!OptionalPluginName.IsEmpty())
+	{
+		TSharedPtr<IPropertyHandle> DlcPluginProperty = NewlyAddedItem->GetChildHandle("AppliesToDlcPlugin");
+		check(DlcPluginProperty.IsValid());
+		DlcPluginProperty->SetValue(OptionalPluginName);
+	}
+	NewlyAddedItem->GetChildHandle("CultureId")->SetValue(CultureId);
+	NewlyAddedItem->NotifyPostChange(EPropertyChangeType::ArrayAdd);
+	OnRebuildChildren.ExecuteIfBound();
+}
+
+void FHoloLensLocalizedResourcesNodeBuilder::EmptyLocalizedEntries()
+{
+	uint32 NumElements;
+	LocalizedResourceCollectionArray->GetNumElements(NumElements);
+
+	for (int32 i = static_cast<int32>(NumElements) - 1; i >= 0; --i)
+	{
+		TSharedRef<IPropertyHandle> Element = LocalizedResourceCollectionArray->GetElement(i);
+		FString AppliesToDlcName;
+		TSharedPtr<IPropertyHandle> DlcPluginProperty = Element->GetChildHandle("AppliesToDlcPlugin");
+		if (DlcPluginProperty.IsValid())
+		{
+			DlcPluginProperty->GetValue(AppliesToDlcName);
+		}
+		FString CultureId;
+		Element->GetChildHandle("CultureId")->GetValue(CultureId);
+		if (AppliesToDlcName == OptionalPluginName && !CultureId.IsEmpty())
+		{
+			LocalizedResourceCollectionArray->DeleteItem(i);
+		}
+	}
+	OnRebuildChildren.ExecuteIfBound();
+}
+
+void FHoloLensLocalizedResourcesNodeBuilder::GenerateChildContent(IDetailChildrenBuilder& ChildrenBuilder)
+{
+	AvailableUECultureNames.Empty();
+	for (const FString& AvailableCulture : GetDefault<UProjectPackagingSettings>()->CulturesToStage)
+	{
+		AvailableUECultureNames.Add(MakeShared<FString>(AvailableCulture));
+	}
+
+	uint32 NumElements;
+	LocalizedResourceCollectionArray->GetNumElements(NumElements);
+	bool HasLocalizedElements = false;
+
+	TSharedPtr<IPropertyHandle> DefaultResources;
+	for (uint32 i = 0; i < NumElements; ++i)
+	{
+		TSharedRef<IPropertyHandle> Element = LocalizedResourceCollectionArray->GetElement(i);
+		FString AppliesToDlcName;
+		TSharedPtr<IPropertyHandle> DlcPluginProperty = Element->GetChildHandle("AppliesToDlcPlugin");
+		if (DlcPluginProperty.IsValid())
+		{
+			DlcPluginProperty->GetValue(AppliesToDlcName);
+		}
+		FString CultureId;
+		Element->GetChildHandle("CultureId")->GetValue(CultureId);
+
+		if (AppliesToDlcName == OptionalPluginName)
+		{
+			if (CultureId.IsEmpty())
+			{
+				DefaultResources = Element;
+			}
+			else
+			{
+				HasLocalizedElements = true;
+				AvailableUECultureNames.RemoveAllSwap(
+					[CultureId](TSharedPtr<FString> Entry)
+				{
+					return *Entry == CultureId;
+				});
+			}
+		}
+	}
+
+	check(DefaultResources.IsValid());
+
+	IDetailGroup& NeutralGroup = ChildrenBuilder.AddGroup("NeutralResources", LOCTEXT("NeutralResources", "Neutral Resources"));
+	NeutralGroup.AddPropertyRow(DefaultResources.ToSharedRef());
+
+	IDetailGroup& LocalizedGroup = ChildrenBuilder.AddGroup("LocalizedResources", LOCTEXT("LocalizedResources", "Localized Resources"));
+
+	if (AvailableUECultureNames.Num() > 0)
+	{
+		LocalizedGroup.AddWidgetRow()
+			.NameContent()
+			[
+				SNew(STextBlock)
+				.Text(LOCTEXT("AddLocalizationTarget", "Add resource localization target"))
+				.Font(IDetailLayoutBuilder::GetDetailFont())
+			]
+			.ValueContent()
+			[
+				SNew(SHorizontalBox)
+				+ SHorizontalBox::Slot()
+				.AutoWidth()
+				[
+					SAssignNew(SelectedCulture, STextComboBox)
+					.OptionsSource(&AvailableUECultureNames)
+					.InitiallySelectedItem(AvailableUECultureNames[0])
+					.Font(IDetailLayoutBuilder::GetDetailFont())
+				]
+				+ SHorizontalBox::Slot()
+				.AutoWidth()
+				[
+					PropertyCustomizationHelpers::MakeAddButton(FSimpleDelegate::CreateSP(this, &FHoloLensLocalizedResourcesNodeBuilder::AddLocalizedEntryForSelectedCulture))
+				]
+			];
+	}
+	else
+	{
+		SelectedCulture.Reset();
+	}
+
+	for (uint32 i = 0; i < NumElements; ++i)
+	{
+		TSharedRef<IPropertyHandle> Element = LocalizedResourceCollectionArray->GetElement(i);
+		FString AppliesToDlcName;
+		TSharedPtr<IPropertyHandle> DlcPluginProperty = Element->GetChildHandle("AppliesToDlcPlugin");
+		if (DlcPluginProperty.IsValid())
+		{
+			DlcPluginProperty->GetValue(AppliesToDlcName);
+		}
+		FString CultureId;
+		Element->GetChildHandle("CultureId")->GetValue(CultureId);
+
+		if (AppliesToDlcName == OptionalPluginName && !CultureId.IsEmpty())
+		{
+			// Disable the default array button (duplicate/delete/etc.) since
+			// we want tighter control over the allowed operations.
+			LocalizedGroup.AddPropertyRow(Element).ShowPropertyButtons(false);
+		}
+	}
+}
+
+TSharedRef<IPropertyTypeCustomization> FHoloLensLocalizedResourcesCustomization::MakeInstance()
+{
+	return MakeShareable(new FHoloLensLocalizedResourcesCustomization);
+}
+
+void FHoloLensLocalizedResourcesCustomization::CustomizeHeader(TSharedRef<class IPropertyHandle> StructPropertyHandle, class FDetailWidgetRow& HeaderRow, IPropertyTypeCustomizationUtils& StructCustomizationUtils)
+{
+	TSharedRef<IPropertyHandle> UECultureNameProperty = StructPropertyHandle->GetChildHandle(FName("CultureId")).ToSharedRef();
+
+	// Add a button to the header to delete this localization target
+	FString CurrentCulture;
+	UECultureNameProperty->GetValue(CurrentCulture);
+	if (!CurrentCulture.IsEmpty())
+	{
+		HeaderRow.NameContent()
+		[
+			SNew(STextBlock)
+			.Text(FText::FromString(CurrentCulture))
+			.Font(IDetailLayoutBuilder::GetDetailFont())
+		]
+		.ValueContent()
+		[
+			SNew(SHorizontalBox)
+			+ SHorizontalBox::Slot()
+			.AutoWidth()
+			[
+				PropertyCustomizationHelpers::MakeDeleteButton(FSimpleDelegate::CreateLambda(
+					[StructPropertyHandle]()
+				{
+					StructPropertyHandle->GetParentHandle()->AsArray()->DeleteItem(StructPropertyHandle->GetIndexInArray());
+				}))
+			]
+		];
+	}
+}
+
+
+void FHoloLensLocalizedResourcesCustomization::CustomizeChildren(TSharedRef<IPropertyHandle> InStructPropertyHandle, class IDetailChildrenBuilder& ChildBuilder, IPropertyTypeCustomizationUtils& StructCustomizationUtils)
+{
+	// Unpack the Strings property
+	TSharedPtr<IPropertyHandle> StringsProperty = InStructPropertyHandle->GetChildHandle("Strings");
+	if (StringsProperty.IsValid())
+	{
+		uint32 NumChildren = 0;
+		StringsProperty->GetNumChildren(NumChildren);
+		for (uint32 i = 0; i < NumChildren; ++i)
+		{
+			ChildBuilder.AddProperty(StringsProperty->GetChildHandle(i).ToSharedRef());
+		}
+	}
+
+	ChildBuilder.AddProperty(InStructPropertyHandle->GetChildHandle("Images").ToSharedRef());
+}
+
+
+#undef LOCTEXT_NAMESPACE
\ No newline at end of file
diff --git a/Engine/Source/Developer/HoloLens/HoloLensPlatformEditor/Private/HoloLensLocalizedResourcesCustomization.h b/Engine/Source/Developer/HoloLens/HoloLensPlatformEditor/Private/HoloLensLocalizedResourcesCustomization.h
new file mode 100644
--- /dev/null
+++ b/Engine/Source/Developer/HoloLens/HoloLensPlatformEditor/Private/HoloLensLocalizedResourcesCustomization.h
@@ -0,0 +1,47 @@
+// Copyright Epic Games, Inc. All Rights Reserved.
+
+#pragma once
+
+#include "PropertyEditorModule.h"
+#include "IPropertyTypeCustomization.h"
+#include "IDetailCustomNodeBuilder.h"
+
+class FHoloLensLocalizedResourcesNodeBuilder : public IDetailCustomNodeBuilder, public TSharedFromThis<FHoloLensLocalizedResourcesNodeBuilder>
+{
+public:
+	FHoloLensLocalizedResourcesNodeBuilder(TSharedRef<IPropertyHandle> InLocalizedResourceCollectionProperty, const FString& InPluginName);
+	virtual ~FHoloLensLocalizedResourcesNodeBuilder();
+
+	/** IDetailCustomNodeBuilder interface */
+	virtual void SetOnRebuildChildren(FSimpleDelegate InOnRebuildChildren) override { OnRebuildChildren = InOnRebuildChildren; }
+	virtual bool RequiresTick() const override { return false; }
+	virtual void Tick(float DeltaTime) override {}
+	virtual void GenerateHeaderRowContent(FDetailWidgetRow& NodeRow) override {}
+	virtual void GenerateChildContent(IDetailChildrenBuilder& ChildrenBuilder) override;
+	virtual bool InitiallyCollapsed() const override { return true; };
+	virtual FName GetName() const override { return FName(TEXT("LocalizedResources")); }
+
+private:
+	void OnNumElementsChanged() { OnRebuildChildren.ExecuteIfBound(); }
+	void AddLocalizedEntryForSelectedCulture();
+	void AddLocalizedEntry(const FString& CultureId);
+	void EmptyLocalizedEntries();
+
+	FString OptionalPluginName;
+	FSimpleDelegate OnRebuildChildren;
+	TSharedRef<class IPropertyHandleArray> LocalizedResourceCollectionArray;
+	TArray<TSharedPtr<FString>> AvailableUECultureNames;
+	TSharedPtr<class STextComboBox> SelectedCulture;
+};
+
+class FHoloLensLocalizedResourcesCustomization : public IPropertyTypeCustomization
+{
+public:
+	static TSharedRef<IPropertyTypeCustomization> MakeInstance();
+
+	// IPropertyTypeCustomization interface
+	virtual void CustomizeHeader(TSharedRef<class IPropertyHandle> StructPropertyHandle, class FDetailWidgetRow& HeaderRow, IPropertyTypeCustomizationUtils& StructCustomizationUtils) override;
+
+	virtual void CustomizeChildren(TSharedRef<IPropertyHandle> InStructPropertyHandle, class IDetailChildrenBuilder& ChildBuilder, IPropertyTypeCustomizationUtils& StructCustomizationUtils) override;
+	// End of IPropertyTypeCustomization interface
+};
\ No newline at end of file
diff --git a/Engine/Source/Developer/HoloLens/HoloLensPlatformEditor/Private/HoloLensPlatformEditorModule.cpp b/Engine/Source/Developer/HoloLens/HoloLensPlatformEditor/Private/HoloLensPlatformEditorModule.cpp
new file mode 100644
--- /dev/null
+++ b/Engine/Source/Developer/HoloLens/HoloLensPlatformEditor/Private/HoloLensPlatformEditorModule.cpp
@@ -0,0 +1,85 @@
+// Copyright Epic Games, Inc. All Rights Reserved.
+
+#include "PropertyEditorModule.h"
+#include "HoloLensTargetSettingsCustomization.h"
+#include "HoloLensLocalizedResourcesCustomization.h"
+#include "HoloLensImageResourcesCustomization.h"
+#include "ISettingsModule.h"
+#include "Modules/ModuleManager.h"
+#include "HoloLensTargetSettings.h"
+
+
+#define LOCTEXT_NAMESPACE "FHoloLensPlatformEditorModule"
+
+/**
+ * Module for the HoloLens platform editor module.
+ */
+class FHoloLensPlatformEditorModule
+	: public IModuleInterface
+{
+public:
+
+	/** Default constructor. */
+	FHoloLensPlatformEditorModule( )
+	{ }
+
+	/** Destructor. */
+	~FHoloLensPlatformEditorModule( )
+	{
+	}
+
+public:
+
+	// IModuleInterface interface
+
+	virtual void StartupModule() override
+	{
+		ISettingsModule* SettingsModule = FModuleManager::GetModulePtr<ISettingsModule>("Settings");
+
+		if (SettingsModule != nullptr)
+		{
+			SettingsModule->RegisterSettings("Project", "Platforms", "HoloLens",
+				LOCTEXT("TargetSettingsName", "HoloLens"),
+				LOCTEXT("TargetSettingsDescription", "Settings for HoloLens"),
+				GetMutableDefault<UHoloLensTargetSettings>()
+			);
+		}
+
+		// register settings detail panel customization
+		FPropertyEditorModule& PropertyModule = FModuleManager::LoadModuleChecked<FPropertyEditorModule>("PropertyEditor");
+		PropertyModule.RegisterCustomClassLayout(
+			"HoloLensTargetSettings",
+			FOnGetDetailCustomizationInstance::CreateStatic(&FHoloLensTargetSettingsCustomization::MakeInstance)
+			);
+
+		PropertyModule.RegisterCustomPropertyTypeLayout(
+			"HoloLensCorePackageLocalizedResources",
+			FOnGetPropertyTypeCustomizationInstance::CreateStatic(&FHoloLensLocalizedResourcesCustomization::MakeInstance)
+			);
+
+		PropertyModule.RegisterCustomPropertyTypeLayout(
+			"HoloLensCorePackageImageResources",
+			FOnGetPropertyTypeCustomizationInstance::CreateStatic(&FHoloLensCorePackageImagesCustomization::MakeInstance)
+			);
+	}
+
+	virtual void ShutdownModule() override
+	{
+		ISettingsModule* SettingsModule = FModuleManager::GetModulePtr<ISettingsModule>("Settings");
+
+		if (SettingsModule != nullptr)
+		{
+			SettingsModule->UnregisterSettings("Project", "Platforms", "HoloLens");
+		}
+	}
+
+private:
+
+};
+
+
+#undef LOCTEXT_NAMESPACE
+
+
+IMPLEMENT_MODULE(FHoloLensPlatformEditorModule, HoloLensPlatformEditor);
+
diff --git a/Engine/Source/Developer/HoloLens/HoloLensPlatformEditor/Private/HoloLensTargetSettings.cpp b/Engine/Source/Developer/HoloLens/HoloLensPlatformEditor/Private/HoloLensTargetSettings.cpp
new file mode 100644
--- /dev/null
+++ b/Engine/Source/Developer/HoloLens/HoloLensPlatformEditor/Private/HoloLensTargetSettings.cpp
@@ -0,0 +1,58 @@
+// Copyright Epic Games, Inc. All Rights Reserved.
+
+#include "HoloLensTargetSettings.h"
+#include "Misc/ConfigCacheIni.h"
+
+
+/* UHoloLensTargetSettings structors
+ *****************************************************************************/
+
+UHoloLensTargetSettings::UHoloLensTargetSettings(const FObjectInitializer& ObjectInitializer)
+	: Super(ObjectInitializer)
+	, MaxTrianglesPerCubicMeter(500.f)
+	, SpatialMeshingVolumeSize(20.f)
+{
+}
+
+void UHoloLensTargetSettings::PostInitProperties()
+{
+	Super::PostInitProperties();
+
+	// Copy across from old location.
+	FString DefaultConfigFile = GetDefaultConfigFilename();
+	if (GConfig->DoesSectionExist(TEXT("/Script/HoloLensTargetPlatform.HoloLensTargetSettings"), DefaultConfigFile))
+	{
+		if (!GConfig->DoesSectionExist(TEXT("/Script/HoloLensPlatformEditor.HoloLensTargetSettings"), DefaultConfigFile))
+		{
+			FString OldTileBackgroundColorHex;
+			FString OldSplashScreenBackgroundColorHex;
+			GConfig->GetString(TEXT("/Script/HoloLensTargetPlatform.HoloLensTargetSettings"), TEXT("TileBackgroundColor"), OldTileBackgroundColorHex, DefaultConfigFile);
+			GConfig->GetString(TEXT("/Script/HoloLensTargetPlatform.HoloLensTargetSettings"), TEXT("SplashScreenBackgroundColor"), OldSplashScreenBackgroundColorHex, DefaultConfigFile);
+			GConfig->GetString(TEXT("/Script/HoloLensTargetPlatform.HoloLensTargetSettings"), TEXT("MinimumPlatformVersion"), MinimumPlatformVersion, DefaultConfigFile);
+			GConfig->GetString(TEXT("/Script/HoloLensTargetPlatform.HoloLensTargetSettings"), TEXT("MaximumPlatformVersionTested"), MaximumPlatformVersionTested, DefaultConfigFile);
+
+			TileBackgroundColor = FColor::FromHex(OldTileBackgroundColorHex);
+			SplashScreenBackgroundColor = FColor::FromHex(OldSplashScreenBackgroundColorHex);
+
+			TryUpdateDefaultConfigFile();
+		}
+		GConfig->EmptySection(TEXT("/Script/HoloLensTargetPlatform.HoloLensTargetSettings"), DefaultConfigFile);
+		if (!GConfig->DoesSectionExist(TEXT("/Script/HoloLensPlatformEditor.HoloLensTargetSettings"), DefaultConfigFile))
+		{
+			TryUpdateDefaultConfigFile();
+		}
+	}
+
+	// Determine if we need to set default capabilities for this project.
+	GConfig->GetBool(TEXT("/Script/HoloLensPlatformEditor.HoloLensTargetSettings"), TEXT("bSetDefaultCapabilities"), bSetDefaultCapabilities, DefaultConfigFile);
+
+	// Adjust the spatial mapping settings if needed
+	if (MaxTrianglesPerCubicMeter < 1.f)
+	{
+		MaxTrianglesPerCubicMeter = 1.f;
+	}
+	if (SpatialMeshingVolumeSize < 1.f)
+	{
+		SpatialMeshingVolumeSize = 1.f;
+	}
+}
diff --git a/Engine/Source/Developer/HoloLens/HoloLensPlatformEditor/Private/HoloLensTargetSettings.h b/Engine/Source/Developer/HoloLens/HoloLensPlatformEditor/Private/HoloLensTargetSettings.h
new file mode 100644
--- /dev/null
+++ b/Engine/Source/Developer/HoloLens/HoloLensPlatformEditor/Private/HoloLensTargetSettings.h
@@ -0,0 +1,160 @@
+// Copyright Epic Games, Inc. All Rights Reserved.
+
+/*=============================================================================
+	HoloLensTargetSettings.h: Declares the UHoloLensTargetSettings class.
+=============================================================================*/
+
+#pragma once
+
+#include "WindowsTargetSettings.h"
+#include "HoloLensLocalizedResources.h"
+#include "HoloLensTargetSettings.generated.h"
+
+/**
+ * Implements the settings for the HoloLens target platform.
+ */
+UCLASS(config = Engine, defaultconfig)
+class UHoloLensTargetSettings
+	: public UObject
+{
+public:
+
+	GENERATED_UCLASS_BODY()
+
+	virtual const TCHAR* GetConfigOverridePlatform() const override
+	{
+		return TEXT("HoloLens");
+	}
+
+	virtual void PostInitProperties() override;
+
+	/**
+	 * When checked, a build that can be run via emulation is added
+	 */
+	UPROPERTY(EditAnywhere, Config, Category="HoloLens", Meta=(DisplayName="Build for HoloLens Emulation"))
+	bool bBuildForEmulation = false;
+
+	/**
+	 * When the box checked the final bundle has binaries of ARM64 OSes.
+	 */
+	UPROPERTY(EditAnywhere, Config, Category="HoloLens", Meta=(DisplayName="Build for HoloLens Device"))
+	bool bBuildForDevice = true;
+
+	UPROPERTY(EditAnywhere, config, Category = "Packaging", AdvancedDisplay, Meta = (DisplayName = "Use Name in App Logo"))
+	bool bUseNameForLogo = true;
+
+	/**
+	 * Controls whether to use the retail Windows Store environment for license checks.  This must be turned on
+	 * when building for submission to the Windows Store, or when sideloading outside of Developer Mode.  Note,
+	 * however, that testing a build with this flag enables requires that the product is listed in the retail
+	 * catalog for the Windows Store.
+	 */
+	UPROPERTY(EditAnywhere, config, Category = "Packaging", AdvancedDisplay, Meta = (DisplayName = "Use Retail Windows Store Environment"))
+	bool bBuildForRetailWindowsStore;
+	
+	UPROPERTY(EditAnywhere, Config, Category = "Packaging")
+	bool bAutoIncrementVersion;
+	
+	UPROPERTY(EditAnywhere, Config, Category = "App Installer")
+	bool bShouldCreateAppInstaller;
+
+	UPROPERTY(EditAnywhere, Config, Category = "App Installer")
+	FString AppInstallerInstallationURL;
+
+	UPROPERTY(EditAnywhere, Config, Category = "App Installer", meta = (ToolTip = "0 will check on every app launch."))
+	int HoursBetweenUpdateChecks = 0;
+
+	UPROPERTY(EditAnywhere, config, Category = "Rendering", meta = (DisplayName = "Enable PIX Profiling"))
+	bool bEnablePIXProfiling;
+
+	UPROPERTY(EditAnywhere, config, Category = Packaging, meta=(HideAlphaChannel))
+	FColor TileBackgroundColor = FColor::FromHex(FString("#000040"));
+
+	UPROPERTY(EditAnywhere, config, Category = Packaging, meta=(HideAlphaChannel))
+	FColor SplashScreenBackgroundColor = FColor::FromHex(FString("#000040"));
+
+	UPROPERTY(EditAnywhere, config, Category = Packaging)
+	TArray<FHoloLensCorePackageLocalizedResources> PerCultureResources;
+
+	/**
+	 * Identifies the device family that your package will target.
+	 */
+	UPROPERTY(EditAnywhere, config, Category = "OS Info")
+	FString TargetDeviceFamily;
+
+	/**
+	 * Minimum version of the HoloLens platform required to run this title.
+	 * It will not be possible to deploy the build on earlier versions.
+	 */
+	UPROPERTY(EditAnywhere, config, Category = "OS Info", AdvancedDisplay, Meta = (DisplayName = "Minimum supported platform version"))
+	FString MinimumPlatformVersion;
+
+	/**
+	 * Specifies the maximum version of the universal platform on which the
+	 * title is known to work as expected.  When deployed to later versions
+	 * the title will experience the runtime behavior of the version given here.
+	 */
+	UPROPERTY(EditAnywhere, config, Category = "OS Info", AdvancedDisplay, Meta = (DisplayName = "Maximum tested platform version"))
+	FString MaximumPlatformVersionTested;
+
+	/** Used by the HoloLens to indicate triangle density when generating meshes. Defaults to 500 per cubic meter */
+	UPROPERTY(EditAnywhere, Config, Category="Spatial Mapping")
+	float MaxTrianglesPerCubicMeter;
+
+	/** Used by the HoloLens to indicate how large (in meters) of a volume to use for generating meshes. Defaults to a 20m cube */
+	UPROPERTY(EditAnywhere, Config, Category = "Spatial Mapping")
+	float SpatialMeshingVolumeSize;
+
+	/** The compiler version to use for this project. May be different to the chosen IDE. */
+	UPROPERTY(EditAnywhere, config, Category = "Toolchain", Meta = (DisplayName = "Compiler Version"))
+	ECompilerVersion CompilerVersion;
+
+	UPROPERTY(EditAnywhere, config, Category = "Toolchain", Meta = (DisplayName = "Windows 10 SDK Version"))
+	FString Windows10SDKVersion;
+
+	/**
+	 * List of supported <Capability><Capability> elements for the application.
+	 */
+	UPROPERTY(EditAnywhere, config, Category = Capabilities)
+	TArray<FString> CapabilityList;
+
+	/**
+	 * List of supported <Capability><DeviceCapability> elements for the application.
+	 */
+	UPROPERTY(EditAnywhere, config, Category = Capabilities)
+	TArray<FString> DeviceCapabilityList;
+
+	/**
+	 * List of supported <Capability><uap:Capability> elements for the application.
+	 */
+	UPROPERTY(EditAnywhere, config, Category = Capabilities)
+	TArray<FString> UapCapabilityList;
+
+	/**
+	 * List of supported <Capability><uap2:Capability> elements for the application.
+	 */
+	UPROPERTY(EditAnywhere, config, Category = Capabilities)
+	TArray<FString> Uap2CapabilityList;
+
+	/**
+	 * Set default capabilities (InternetClientServer and PrivateNetworkClientServer) for the application.
+	 */
+	UPROPERTY(EditAnywhere, config, Category = Capabilities)
+	bool bSetDefaultCapabilities = true;
+
+	/** Which of the currently enabled spatialization plugins to use. */
+	UPROPERTY(config, EditAnywhere, Category = "Audio")
+	FString SpatializationPlugin;
+
+	/** Which of the currently enabled reverb plugins to use. */
+	UPROPERTY(config, EditAnywhere, Category = "Audio")
+	FString ReverbPlugin;
+
+	/** Which of the currently enabled occlusion plugins to use. */
+	UPROPERTY(config, EditAnywhere, Category = "Audio")
+	FString OcclusionPlugin;
+
+	/** Quality Level to COOK SoundCues at (if set, all other levels will be stripped by the cooker). */
+	UPROPERTY(config, EditAnywhere, Category = "Audio|CookOverrides", meta = (DisplayName = "Sound Cue Cook Quality"))
+	int32 SoundCueCookQualityIndex = INDEX_NONE;
+};
diff --git a/Engine/Source/Developer/HoloLens/HoloLensPlatformEditor/Private/HoloLensTargetSettingsCustomization.cpp b/Engine/Source/Developer/HoloLens/HoloLensPlatformEditor/Private/HoloLensTargetSettingsCustomization.cpp
new file mode 100644
--- /dev/null
+++ b/Engine/Source/Developer/HoloLens/HoloLensPlatformEditor/Private/HoloLensTargetSettingsCustomization.cpp
@@ -0,0 +1,771 @@
+// Copyright Epic Games, Inc. All Rights Reserved.
+
+#include "HoloLensTargetSettingsCustomization.h"
+#include "HoloLensPlatformEditor.h"
+#include "SExternalImageReference.h"
+#include "IExternalImagePickerModule.h"
+#include "DetailLayoutBuilder.h"
+#include "DetailCategoryBuilder.h"
+#include "DetailWidgetRow.h"
+#include "PropertyHandle.h"
+#include "EditorDirectories.h"
+#include "Widgets/Input/SFilePathPicker.h"
+#include "Widgets/Input/STextComboBox.h"
+#include "ISourceControlModule.h"
+#include "SourceControlHelpers.h"
+#include "Widgets/Notifications/SNotificationList.h"
+#include "Framework/Notifications/NotificationManager.h"
+#include "Misc/FileHelper.h"
+#include "Windows/AllowWindowsPlatformTypes.h"
+#include "Wincrypt.h"
+#include "HoloLensTargetSettings.h"
+#include "GeneralProjectSettings.h"
+#include "Widgets/Input/SButton.h"
+#include "Widgets/Notifications/SErrorHint.h"
+#include "HAL/PlatformFileManager.h"
+#include "PropertyCustomizationHelpers.h"
+#include "IDetailChildrenBuilder.h"
+#include "HoloLensLocalizedResourcesCustomization.h"
+
+#include <winreg.h>
+
+#pragma comment( lib, "Advapi32.lib" )
+
+#pragma warning( disable : 6387 )
+
+#define LOCTEXT_NAMESPACE "HoloLensTargetSettingsCustomization"
+
+TArray<FHoloLensSDKVersion> FHoloLensSDKVersion::GetSDKVersions()
+{
+	bool bFound18362 = false;
+
+	static TArray<FHoloLensSDKVersion> SDKVersions;
+	if (SDKVersions.Num() == 0)
+	{
+		HKEY hRoots[] = { HKEY_LOCAL_MACHINE , HKEY_CURRENT_USER };
+
+		for (DWORD dwRootIdx = 0; dwRootIdx < _countof(hRoots); ++dwRootIdx)
+		{
+			HKEY hKey;
+			if (RegOpenKeyExW(hRoots[dwRootIdx], L"Software\\Microsoft\\Windows Kits\\Installed Roots", 0, KEY_ENUMERATE_SUB_KEYS | KEY_QUERY_VALUE, &hKey) == ERROR_SUCCESS)
+			{
+				DWORD dwIndex = 0;
+				WCHAR sSubKeyName[256];
+				DWORD dwSubKeyNameLen = _countof(sSubKeyName);
+				FString SDKPath;
+
+				if (ERROR_SUCCESS == RegQueryValueExW(hKey, L"KitsRoot10", 0, NULL, (LPBYTE)sSubKeyName, &dwSubKeyNameLen))
+				{
+					SDKPath = FString(dwSubKeyNameLen, sSubKeyName);
+				}
+				dwSubKeyNameLen = _countof(sSubKeyName);
+
+				while (RegEnumKeyExW(hKey, dwIndex, sSubKeyName, &dwSubKeyNameLen, NULL, NULL, NULL, NULL) == ERROR_SUCCESS)
+				{
+					int version[4];
+					if (swscanf_s(sSubKeyName, L"%d.%d.%d.%d", &version[0], &version[1], &version[2], &version[3]) == 4)
+					{
+						if (FPaths::FileExists(FPaths::Combine(SDKPath, TEXT("Include"), FString(dwSubKeyNameLen, sSubKeyName), TEXT("um"), TEXT("Windows.h"))))
+						{
+							// HoloLens min version is 17763 for remoting and 18317 for hand tracking
+							// @todo JoeG - update this version when the official os/sdk are released
+							if (version[0] >= 10 && version[2] >= 17763)
+							{
+								new(SDKVersions) FHoloLensSDKVersion(FString(dwSubKeyNameLen, sSubKeyName), version[0], version[1], version[2], version[3]);
+							}
+							if (version[2] == 18362)
+							{
+								bFound18362 = true;
+							}
+						}
+					}
+					dwSubKeyNameLen = _countof(sSubKeyName);
+					dwIndex++;
+				}
+				RegCloseKey(hKey);
+			}
+		}
+
+		// 18362 is required as the official minimum SDK released by Microsoft for all needed HL2 features to work
+		// If it's not found, we still need it as a hidden SDK version so that we can set MinimumPlatformVersion correctly to
+		// allow the package to run on lower installed firmware revisions
+		if (!bFound18362)
+		{
+			new(SDKVersions) FHoloLensSDKVersion(FString(L"10.0.18362.0"), 10, 0, 18362, 0);
+		}
+	}
+
+	return SDKVersions;
+}
+
+TSharedRef<IDetailCustomization> FHoloLensTargetSettingsCustomization::MakeInstance()
+{
+	return MakeShareable(new FHoloLensTargetSettingsCustomization);
+}
+
+void FHoloLensTargetSettingsCustomization::InitSupportedPlatformVersions()
+{
+	PlatformVersionOptions.Empty();
+
+	const TArray<FHoloLensSDKVersion>& SDKVersions = FHoloLensSDKVersion::GetSDKVersions();
+	for (const FHoloLensSDKVersion& SDKVer : SDKVersions)
+	{
+		PlatformVersionOptions.Add(MakeShareable(new FString(SDKVer.VersionString)));
+	}
+}
+
+void FHoloLensTargetSettingsCustomization::InitTargetDeviceFamilyOptions()
+{
+	TargetDeviceFamilyOptions.Empty();
+
+	TargetDeviceFamilyOptions.Add(MakeShareable(new FString("Windows.Holographic")));
+}
+
+void FHoloLensTargetSettingsCustomization::CustomizeDetails(IDetailLayoutBuilder& DetailBuilder)
+{
+	InitTargetDeviceFamilyOptions();
+	InitSupportedPlatformVersions();
+
+	// Add UI for selecting the TargetDeviceFamily.
+	TSharedRef<IPropertyHandle> TargetDeviceFamilyProperty = DetailBuilder.GetProperty("TargetDeviceFamily");
+	AddWidgetForTargetDeviceFamily(DetailBuilder, TargetDeviceFamilyProperty);
+
+	// Add UI for select min/max platform version.
+	TSharedRef<IPropertyHandle> MinVersionProperty = DetailBuilder.GetProperty("MinimumPlatformVersion");
+	// 10.0.18362.0 is the minimum 'officially' released SDK from MS that supports all the HoloLens 2 features that are needed
+	AddWidgetForPlatformVersion(DetailBuilder, MinVersionProperty, nullptr, true, FString(L"10.0.18362.0"));
+	TSharedRef<IPropertyHandle> MaxVersionProperty = DetailBuilder.GetProperty("MaximumPlatformVersionTested");
+	AddWidgetForPlatformVersion(DetailBuilder, MaxVersionProperty, nullptr, true);
+
+	// Add UI to select signing certificate
+	IDetailCategoryBuilder& PackagingCategoryBuilder = DetailBuilder.EditCategory(FName("Packaging"));
+	TSharedRef<IPropertyHandle> LocalizedResourcesProperty = DetailBuilder.GetProperty("PerCultureResources");
+	DetailBuilder.HideProperty(LocalizedResourcesProperty);
+	PackagingCategoryBuilder.AddCustomBuilder(MakeShared<FHoloLensLocalizedResourcesNodeBuilder>(LocalizedResourcesProperty, TEXT("")));
+
+	FString ProjectPath = FPaths::ProjectDir() / TEXT("Build") / TEXT("HoloLens") / TEXT("SigningCertificate.pfx");
+
+	// Load the existing signing certificate (if any)
+	FText SigningCertificateCaption = LOCTEXT("SigningCertificate", "Signing Certificate");
+	FText SigningCertificateTooltip = LOCTEXT("SigningCertificateTooltip", "Pfx file containing a private key used to sign the AppX file created during packaging.  The certificate subject name must exactly match the value of Package/Identity/Publisher.  Signing is required when sideloading packaged builds");
+
+	PackagingCategoryBuilder.AddCustomRow(SigningCertificateCaption)
+	.NameContent()
+	[
+		SNew(STextBlock)
+		.Text(SigningCertificateCaption)
+		.ToolTipText(SigningCertificateTooltip)
+		.Font(IDetailLayoutBuilder::GetDetailFont())
+	]
+	.ValueContent()
+	.MaxDesiredWidth(500.0f)
+	.MinDesiredWidth(100.0f)
+	[
+		SNew(SHorizontalBox)
+		+ SHorizontalBox::Slot()
+		.FillWidth(1.0f)
+		.VAlign(VAlign_Center)
+		[
+			SNew(SFilePathPicker)
+			.BrowseButtonImage(FEditorStyle::GetBrush("PropertyWindow.Button_Ellipsis"))
+			.BrowseButtonStyle(FEditorStyle::Get(), "HoverHintOnly")
+			.BrowseDirectory(FEditorDirectories::Get().GetLastDirectory(ELastDirectory::GENERIC_OPEN))
+			.FileTypeFilter(TEXT("pfx files (*.pfx)|*.pfx"))
+			.IsReadOnly(true)
+			.OnPathPicked(this, &FHoloLensTargetSettingsCustomization::OnCertificatePicked)
+			.FilePath(this, &FHoloLensTargetSettingsCustomization::GetSigningCertificateSubjectName)
+			.Font(IDetailLayoutBuilder::GetDetailFont())
+		]
+		+ SHorizontalBox::Slot()
+		.AutoWidth()
+		[
+			SAssignNew(SigningCertificateError, SErrorHint)
+		]
+		+ SHorizontalBox::Slot()
+		.AutoWidth()
+		.VAlign(VAlign_Center)
+		[
+			SNew(SButton)
+			.Text(LOCTEXT("NewSigningCertificate", "Generate new"))
+			.OnClicked(this, &FHoloLensTargetSettingsCustomization::GenerateSigningCertificate)
+		]
+	];
+
+	LoadAndValidateSigningCertificate();
+
+	IDetailCategoryBuilder& ToolchainCategoryBuilder = DetailBuilder.EditCategory(FName("Toolchain"));
+
+	// Make sure compiler version comes first - it's probably most important
+	ToolchainCategoryBuilder.AddProperty(FName("CompilerVersion"));
+
+	TSharedRef<IPropertyHandle> Win10SDKVersionPropertyHandle = DetailBuilder.GetProperty("Windows10SDKVersion");
+
+	FText AutoDetectSDKCaption = LOCTEXT("AutodetectWin10SDKCaption", "Auto-detect Windows 10 SDK");
+	FText AutoDetectSDKTooltip = LOCTEXT("AutodetectWin10SDKTooltip", "When enabled the project will build against the most recent version of the Windows 10 SDK supported by your compiler. This is typically the recommended behavior.  Uncheck in order to manually select a specific SDK version");
+	TSharedPtr<SCheckBox> AutoDetectSDKCheckbox;
+	ToolchainCategoryBuilder.AddCustomRow(AutoDetectSDKCaption)
+	.NameContent()
+	[
+		SNew(STextBlock)
+		.Text(AutoDetectSDKCaption)
+		.ToolTipText(AutoDetectSDKTooltip)
+		.Font(IDetailLayoutBuilder::GetDetailFont())
+	]
+	.ValueContent()
+	.MaxDesiredWidth(500.0f)
+	.MinDesiredWidth(100.0f)
+	[
+		SAssignNew(AutoDetectSDKCheckbox, SCheckBox)
+		.IsChecked(this, &FHoloLensTargetSettingsCustomization::IsAutoDetectWin10SDKChecked)
+		.OnCheckStateChanged(this, &FHoloLensTargetSettingsCustomization::OnAutoDetectWin10SDKChanged, Win10SDKVersionPropertyHandle)
+	];
+
+	AddWidgetForPlatformVersion(DetailBuilder, Win10SDKVersionPropertyHandle, &WindowsSDKSelector);
+
+	// Add capability support.
+	TSharedRef<IPropertyHandle> CapabilityList = DetailBuilder.GetProperty("CapabilityList");
+	DetailBuilder.HideProperty(CapabilityList);
+	AddWidgetForCapability(DetailBuilder, CapabilityList, TEXT("internetClient"), LOCTEXT("InternetClientCaption", "Internet Client"), LOCTEXT("InternetClientTooltip", "Provides outbound access to the Internet and networks in public places like airports and cofee shops."), false);
+	AddWidgetForCapability(DetailBuilder, CapabilityList, TEXT("internetClientServer"), LOCTEXT("InternetClientServerCaption", "Internet Client Server"), LOCTEXT("InternetClientServerTooltip", "Provides inbound and outbound access to the Internet and networks in public places like airports and cofee shops"), false);
+	AddWidgetForCapability(DetailBuilder, CapabilityList, TEXT("privateNetworkClientServer"), LOCTEXT("PrivateNetworkCaption", "Private Network Client Server"), LOCTEXT("PrivateNetworkTooltip", "Provides inbound and outbound access to the Internet and networks that have an authenticated domain controller, or that the user has designated as either home or work networks.Inbound access to critical ports is always blocked."), false);
+	AddWidgetForCapability(DetailBuilder, CapabilityList, TEXT("allJoyn"), LOCTEXT("AllJoynCaption", "AllJoyn"), LOCTEXT("AllJoynTooltip", "Allows AllJoyn-enabled apps and devices on a network to discover and interact with each other."), true);
+	AddWidgetForCapability(DetailBuilder, CapabilityList, TEXT("codeGeneration"), LOCTEXT("CodeGenCaption", "Code Generation"), LOCTEXT("CodeGenTooltip", "Allows apps to generate code dynamically."), true);
+
+	TSharedRef<IPropertyHandle> DeviceCapabilityList = DetailBuilder.GetProperty("DeviceCapabilityList");
+	DetailBuilder.HideProperty(DeviceCapabilityList);
+	AddWidgetForCapability(DetailBuilder, DeviceCapabilityList, TEXT("microphone"), LOCTEXT("MicrophoneCaption", "Microphone"), LOCTEXT("MicrophoneTooltip", "Provides access to the microphone's audio feed, which allows the app to record audio from connected microphones. Required for Windows.Media.SpeechRecognition APIs."), false);
+	AddWidgetForCapability(DetailBuilder, DeviceCapabilityList, TEXT("webcam"), LOCTEXT("WebcamCaption", "Webcam"), LOCTEXT("WebcamTooltip", "Provides access to the webcam's video feed, which allows the app to capture snapshots and movies from connected webcams."), false);
+	AddWidgetForCapability(DetailBuilder, DeviceCapabilityList, TEXT("proximity"), LOCTEXT("ProximityCaption", "Proximity"), LOCTEXT("ProximityTooltip", "Provides capability to connect devices in close proximity to the PC via near field proximity radio or Wi-FI Direct."), true);
+	AddWidgetForCapability(DetailBuilder, DeviceCapabilityList, TEXT("location"), LOCTEXT("LocationCaption", "Location"), LOCTEXT("LocationTooltip", "Provides access to the current location, which is obtained from dedicated hardware like a GPS sensor in the PC or derived from available network information."), true);
+	AddWidgetForCapability(DetailBuilder, DeviceCapabilityList, TEXT("bluetooth"), LOCTEXT("BluetoothCaption", "Bluetooth"), LOCTEXT("BluetoothTooltip", "Allows communication with paired Bluetooth devices over the Generic Attribute (GATT) or Classic Basic Rate (RFCOMM) protocols."), true);
+	AddWidgetForCapability(DetailBuilder, DeviceCapabilityList, TEXT("gazeInput"), LOCTEXT("GazeInputCaption", "Gaze Input"), LOCTEXT("GazeInputTooltip", "Provides access to the input from gaze like eyes tracking. Required for Windows.Perception.People.EyesPose API."), false);
+	AddWidgetForCapability(DetailBuilder, DeviceCapabilityList, TEXT("wiFiControl"), LOCTEXT("WiFiControlCaption", "WiFiControl"), LOCTEXT("WiFiControlTooltip", "Allows apps to scan and connect to Wi-Fi networks."), false);
+
+	TSharedRef<IPropertyHandle> UapCapabilityList = DetailBuilder.GetProperty("UapCapabilityList");
+	DetailBuilder.HideProperty(UapCapabilityList);
+	AddWidgetForCapability(DetailBuilder, UapCapabilityList, TEXT("musicLibrary"), LOCTEXT("MusicLibCaption", "Music Library"), LOCTEXT("MusicLibTooltip", "Provides capability to add, change, or delete files in the Music Library for the local PC and HomeGroup PCs."), true);
+	AddWidgetForCapability(DetailBuilder, UapCapabilityList, TEXT("picturesLibrary"), LOCTEXT("PicturesLibCaption", "Pictures Library"), LOCTEXT("PicturesLibTooltip", "Provides capability to add, change, or delete files in the Pictures Library for the local PC and HomeGroup PCs."), true);
+	AddWidgetForCapability(DetailBuilder, UapCapabilityList, TEXT("videosLibrary"), LOCTEXT("VideosLibCaption", "Videos Library"), LOCTEXT("VideosLibTooltip", "Provides capability to add, change, or delete files in the Videos Library for the local PC and HomeGroup PCs."), true);
+	AddWidgetForCapability(DetailBuilder, UapCapabilityList, TEXT("blockedChatMessages"), LOCTEXT("BlockedChatCaption", "Blocked Chat Messages"), LOCTEXT("BlockedChatTooltip", "Allows apps to read SMS and MMS messages that have been blocked by the Spam Filter app."), true);
+	AddWidgetForCapability(DetailBuilder, UapCapabilityList, TEXT("chat"), LOCTEXT("ChatCaption", "Chat"), LOCTEXT("ChatTooltip", "Allows apps to read and delete Text Messages. It also allows apps to store chat messages in the system data store."), true);
+	AddWidgetForCapability(DetailBuilder, UapCapabilityList, TEXT("enterpriseAuthentication"), LOCTEXT("EnterpriseAuthCaption", "Enterprise Authentication"), LOCTEXT("EnterpriseAuthTooltip", "Subject to Store policy. Provides ability to connect to enterprise intranet resources that require domain credentials.This capability is typically not needed for most apps."), true);
+	AddWidgetForCapability(DetailBuilder, UapCapabilityList, TEXT("objects3D"), LOCTEXT("Objects3DCaption", "Objects 3D"), LOCTEXT("Objects3DTooltip", "Provides access to the user's 3D Objects, allowing the app to enumerate and access all files in the library without user interaction."), true);
+	AddWidgetForCapability(DetailBuilder, UapCapabilityList, TEXT("phoneCall"), LOCTEXT("PhoneCallCaption", "Phone Call"), LOCTEXT("PhoneCallTooltip", "Allows apps to access all phone lines on the device and perform the following functions: place a call, access line-related metadata, access line-related triggers, set and check block list and call origination information."), true);
+	AddWidgetForCapability(DetailBuilder, UapCapabilityList, TEXT("removableStorage"), LOCTEXT("RemovableStorageCaption", "Removable Storage"), LOCTEXT("RemovableStorageTooltip", "Provides capability to add, change, or delete files on removable storage devices."), true);
+	AddWidgetForCapability(DetailBuilder, UapCapabilityList, TEXT("sharedUserCertificates"), LOCTEXT("ShareCertsCaption", "Shared User Certificates"), LOCTEXT("SharedCertsTooltip", "Subject to Store policy. Provides capability to access software and hardware certificates for validating a user's identity."), true);
+	AddWidgetForCapability(DetailBuilder, UapCapabilityList, TEXT("userAccountInformation"), LOCTEXT("UserAccountCaption", "User Account Information"), LOCTEXT("UserAccountTooltip", "Gives apps the ability to access user's name and picture. Required to access Windows.System.UserProfile APIs."), true);
+	AddWidgetForCapability(DetailBuilder, UapCapabilityList, TEXT("voipCall"), LOCTEXT("VoipCallCaption", "VOIP Call"), LOCTEXT("VoipCallTooltip", "Allows access to the VOIP calling APIs in Windows.ApplicationModel.Calls."), true);
+
+	TSharedRef<IPropertyHandle> Uap2CapabilityList = DetailBuilder.GetProperty("Uap2CapabilityList");
+	DetailBuilder.HideProperty(Uap2CapabilityList);
+	AddWidgetForCapability(DetailBuilder, Uap2CapabilityList, TEXT("spatialPerception"), LOCTEXT("SpatialPerceptionCaption", "Spatial Perception"), LOCTEXT("SpatialPerceptionTooltip", "Provides access to environment data, which will be used to generate spatial maps or stages. Required to access Windows.Perception.Spatial APIs."), false);
+
+	// If this is the first time capabilities are being accessed for the project, enable defaults.
+	TSharedRef<IPropertyHandle> SetDefaultCapabilitiesProperty = DetailBuilder.GetProperty("bSetDefaultCapabilities");
+	DetailBuilder.HideProperty(SetDefaultCapabilitiesProperty);
+	bool bSetDefaults;
+	SetDefaultCapabilitiesProperty->GetValue(bSetDefaults);
+	if (bSetDefaults)
+	{
+		OnCapabilityStateChanged(ECheckBoxState::Checked, CapabilityList, TEXT("internetClientServer"));
+		OnCapabilityStateChanged(ECheckBoxState::Checked, CapabilityList, TEXT("privateNetworkClientServer"));
+		// Our AR apps expect to get spatial mesh updates, so make this opt out for users that know what they are doing
+		OnCapabilityStateChanged(ECheckBoxState::Checked, Uap2CapabilityList, TEXT("spatialPerception"));
+
+		SetDefaultCapabilitiesProperty->NotifyPreChange();
+		SetDefaultCapabilitiesProperty->SetValue(false);
+		SetDefaultCapabilitiesProperty->NotifyPostChange(EPropertyChangeType::ValueSet);
+	}
+
+	AudioPluginManager.BuildAudioCategory(DetailBuilder, TEXT("HoloLens"));
+}
+
+void FHoloLensTargetSettingsCustomization::OnCertificatePicked(const FString& PickedPath)
+{
+	FString CertificatePath = FPaths::ProjectDir() / TEXT("Build") / TEXT("HoloLens") / TEXT("SigningCertificate.pfx");
+	FText FailReason;
+	if (!SourceControlHelpers::CopyFileUnderSourceControl(CertificatePath, PickedPath, LOCTEXT("CertificateDescription", "Certificate"), FailReason))
+	{
+		FNotificationInfo Info(FailReason);
+		Info.ExpireDuration = 3.0f;
+		FSlateNotificationManager::Get().AddNotification(Info);
+	}
+
+	LoadAndValidateSigningCertificate();
+}
+
+FString FHoloLensTargetSettingsCustomization::GetNameForSigningCertificate(const FString &CertificatePath)
+{
+	// Display the friendly name of the certificate instead of the file name.
+	FString CertificateName = LOCTEXT("CertificateNoCertificate", "No valid certificate selected").ToString();
+	TArray<uint8> CertBytes;
+	if (FFileHelper::LoadFileToArray(CertBytes, *CertificatePath))
+	{
+		CRYPT_DATA_BLOB CertBlob;
+		CertBlob.cbData = CertBytes.Num();
+		CertBlob.pbData = CertBytes.GetData();
+		HCERTSTORE CertStore = PFXImportCertStore(&CertBlob, nullptr, 0);
+		if (CertStore != nullptr)
+		{
+			DWORD StrType = CERT_X500_NAME_STR;
+			const CERT_CONTEXT *CertContext = CertEnumCertificatesInStore(CertStore, nullptr);
+			if (CertContext != nullptr)
+			{
+				DWORD NumCharacters = CertGetNameString(CertContext, CERT_NAME_RDN_TYPE, 0, &StrType, nullptr, 0);
+				CertificateName.GetCharArray().SetNumZeroed(NumCharacters);
+				CertGetNameString(CertContext, CERT_NAME_RDN_TYPE, 0, &StrType, CertificateName.GetCharArray().GetData(), CertificateName.GetAllocatedSize());
+				CertFreeCertificateContext(CertContext);
+			}
+			CertCloseStore(CertStore, 0);
+		}
+		else
+		{
+			DWORD FailureReason = GetLastError();
+			switch (FailureReason)
+			{
+			case ERROR_INVALID_PASSWORD:
+				CertificateName = LOCTEXT("CertificatePasswordProtected", "Certificate is password protected").ToString();
+				break;
+
+			default:
+				break;
+			}
+		}
+	}
+
+	return CertificateName;
+}
+
+void FHoloLensTargetSettingsCustomization::AddWidgetForPlatformVersion(IDetailLayoutBuilder& DetailBuilder, TSharedRef<IPropertyHandle> PropertyHandle, TSharedPtr<STextComboBox>* OutVersionSelector, bool bForAdvanced, const FString& SpecificVersionOverrideString)
+{
+	IDetailCategoryBuilder& VersionCategoryBuilder = DetailBuilder.EditCategory(FName(*PropertyHandle->GetMetaData("Category")));
+	DetailBuilder.HideProperty(PropertyHandle);
+
+	FString CurrentSelectedVersion = SpecificVersionOverrideString;
+	PropertyHandle->GetValue(CurrentSelectedVersion);
+
+	// Default to latest version when not set explicitly.  With Windows 10 automatic updates
+	// it's a reasonably safe choice for users, and much more developer friendly than defaulting
+	// to the oldest version.
+	int32 CurrentSelectedIndex = PlatformVersionOptions.Num() - 1;
+	if (CurrentSelectedVersion.IsEmpty())
+	{
+		OnSelectedItemChanged(PlatformVersionOptions[CurrentSelectedIndex], ESelectInfo::Direct, PropertyHandle);
+	}
+	else
+	{
+		for (int32 i = 0; i < PlatformVersionOptions.Num(); ++i)
+		{
+			if (*PlatformVersionOptions[i] == CurrentSelectedVersion)
+			{
+				CurrentSelectedIndex = i;
+				break;
+			}
+		}
+	}
+
+	TSharedPtr<STextComboBox> Unused;
+	TSharedPtr<STextComboBox>& VersionSelectorRef = OutVersionSelector != nullptr ? *OutVersionSelector : Unused;
+
+	VersionCategoryBuilder.AddCustomRow(PropertyHandle->GetPropertyDisplayName(), bForAdvanced)
+	.NameContent()
+	[
+		PropertyHandle->CreatePropertyNameWidget()
+	]
+	.ValueContent()
+	.MaxDesiredWidth(500.0f)
+	.MinDesiredWidth(100.0f)
+	[
+		SNew(SHorizontalBox)
+		+ SHorizontalBox::Slot()
+		.FillWidth(1.0f)
+		.VAlign(VAlign_Center)
+		[
+			SAssignNew(VersionSelectorRef, STextComboBox)
+			.Font(IDetailLayoutBuilder::GetDetailFont())
+			.OptionsSource(&PlatformVersionOptions)
+			.InitiallySelectedItem(PlatformVersionOptions[CurrentSelectedIndex])
+			.OnSelectionChanged(this, &FHoloLensTargetSettingsCustomization::OnSelectedItemChanged, PropertyHandle)
+		]
+	];
+}
+
+void FHoloLensTargetSettingsCustomization::AddWidgetForTargetDeviceFamily(IDetailLayoutBuilder& DetailBuilder, TSharedRef<IPropertyHandle> PropertyHandle)
+{
+	IDetailCategoryBuilder& TargetCategoryBuilder = DetailBuilder.EditCategory(FName(*PropertyHandle->GetMetaData("Category")));
+	DetailBuilder.HideProperty(PropertyHandle);
+
+	FString CurrentSelectedDeviceFamily;
+	PropertyHandle->GetValue(CurrentSelectedDeviceFamily);
+
+	// Default to first option when not set explicitly.
+	int32 CurrentSelectedIndex = 0;
+	if (CurrentSelectedDeviceFamily.IsEmpty())
+	{
+		OnSelectedItemChanged(TargetDeviceFamilyOptions[CurrentSelectedIndex], ESelectInfo::Direct, PropertyHandle);
+	}
+	else
+	{
+		for (int32 i = 0; i < TargetDeviceFamilyOptions.Num(); i++)
+		{
+			if (*TargetDeviceFamilyOptions[i] == CurrentSelectedDeviceFamily)
+			{
+				CurrentSelectedIndex = i;
+				break;
+			}
+		}
+	}
+
+	TargetCategoryBuilder.AddCustomRow(PropertyHandle->GetPropertyDisplayName())
+	.NameContent()
+	[
+		PropertyHandle->CreatePropertyNameWidget()
+	]
+	.ValueContent()
+	.MaxDesiredWidth(500.0f)
+	.MinDesiredWidth(100.0f)
+	[
+		SNew(SHorizontalBox)
+		+ SHorizontalBox::Slot()
+		.FillWidth(1.0f)
+		.VAlign(VAlign_Center)
+		[
+			SNew(STextComboBox)
+			.Font(IDetailLayoutBuilder::GetDetailFont())
+			.OptionsSource(&TargetDeviceFamilyOptions)
+			.InitiallySelectedItem(TargetDeviceFamilyOptions[CurrentSelectedIndex])
+			.OnSelectionChanged(this, &FHoloLensTargetSettingsCustomization::OnSelectedItemChanged, PropertyHandle)
+		]
+	];
+}
+
+void FHoloLensTargetSettingsCustomization::OnSelectedItemChanged(TSharedPtr<FString> NewValue, ESelectInfo::Type SelectInfo, TSharedRef<IPropertyHandle> Handle)
+{
+	Handle->SetValue(*NewValue);
+}
+
+void FHoloLensTargetSettingsCustomization::AddWidgetForCapability(IDetailLayoutBuilder& DetailBuilder, TSharedRef<IPropertyHandle> CapabilityList, const FString& CapabilityName, const FText& CapabilityCaption, const FText& CapabilityTooltip, bool bForAdvanced)
+{
+	IDetailCategoryBuilder& CapabilityBuilder = DetailBuilder.EditCategory(FName("Capabilities"));
+
+	// Initialize checkbox state based on whether or not the capability currently exists in the CapabilityList.
+	ECheckBoxState currentState = IsCapabilityChecked(CapabilityList, CapabilityName);
+	if (currentState == ECheckBoxState::Checked)
+	{
+		OnCapabilityStateChanged(currentState, CapabilityList, CapabilityName);
+	}
+
+	CapabilityBuilder.AddCustomRow(CapabilityCaption, bForAdvanced)
+	.NameContent()
+	[
+		SNew(STextBlock)
+		.Text(CapabilityCaption)
+		.ToolTipText(CapabilityTooltip)
+		.Font(IDetailLayoutBuilder::GetDetailFont())
+	]
+	.ValueContent()
+	.VAlign(VAlign_Center)
+	[
+		SNew(SCheckBox)
+		.IsChecked(this, &FHoloLensTargetSettingsCustomization::IsCapabilityChecked, CapabilityList, CapabilityName)
+		.OnCheckStateChanged(this, &FHoloLensTargetSettingsCustomization::OnCapabilityStateChanged, CapabilityList, CapabilityName)
+	];
+}
+
+ECheckBoxState FHoloLensTargetSettingsCustomization::IsCapabilityChecked(TSharedRef<IPropertyHandle> CapabilityList, const FString CapabilityName) const
+{
+	TArray<void*> RawData;
+	CapabilityList->AccessRawData(RawData);
+	TArray<FString>* RawCapabilityStringArray = reinterpret_cast<TArray<FString>*>(RawData[0]);
+
+	int32 Index;
+	bool Found = RawCapabilityStringArray->Find(CapabilityName, Index);
+	return (Found ? ECheckBoxState::Checked : ECheckBoxState::Unchecked);
+}
+
+void FHoloLensTargetSettingsCustomization::OnCapabilityStateChanged(ECheckBoxState CheckState, TSharedRef<IPropertyHandle> CapabilityList, const FString CapabilityName)
+{
+	bool IsEnabled = (CheckState == ECheckBoxState::Checked);
+	TArray<void*> RawData;
+	CapabilityList->AccessRawData(RawData);
+	TArray<FString>* RawCapabilityStringArray = reinterpret_cast<TArray<FString>*>(RawData[0]);
+
+	CapabilityList->NotifyPreChange();
+	int32 Index;
+	bool Found = RawCapabilityStringArray->Find(CapabilityName, Index);
+
+	if (Found && !IsEnabled)
+	{
+		// Remove existing capability from the list
+		RawCapabilityStringArray->RemoveAt(Index);
+		CapabilityList->NotifyPostChange(EPropertyChangeType::ArrayRemove);
+	}
+	else if (!Found && IsEnabled)
+	{
+		//Add new capability to the list
+		RawCapabilityStringArray->AddUnique(*CapabilityName);
+		CapabilityList->NotifyPostChange(EPropertyChangeType::ArrayAdd);
+	}
+}
+
+FReply FHoloLensTargetSettingsCustomization::GenerateSigningCertificate()
+{
+	FString WinSDKSubKey = TEXT("SOFTWARE\\Wow6432Node\\Microsoft\\Microsoft SDKs\\Windows\\v10.0");
+	FString WinSDKFolder;
+	if (!FPlatformMisc::QueryRegKey(HKEY_CURRENT_USER, *WinSDKSubKey, TEXT("InstallationFolder"), WinSDKFolder))
+	{
+		if (!FPlatformMisc::QueryRegKey(HKEY_LOCAL_MACHINE, *WinSDKSubKey, TEXT("InstallationFolder"), WinSDKFolder))
+		{
+			FNotificationInfo Info(LOCTEXT("SigningCertificateFailed_SDKInstallationFolder", "Failed to generate certificate: could not find Windows 10 SDK Installation Folder"));
+			Info.ExpireDuration = 3.0f;
+			FSlateNotificationManager::Get().AddNotification(Info);
+
+			return FReply::Handled();
+		}
+	}
+
+	//make sure the company name has required syntax
+	FString CompanyDistinguishedName = GetDefault<UGeneralProjectSettings>()->CompanyDistinguishedName;
+	if (!CompanyDistinguishedName.Contains(TEXT("CN=")))
+	{
+		FNotificationInfo Info(LOCTEXT("SigningCertificateFailed_MalformedCompanyName", "Company Distinguished Name must contain 'CN=YourCompany' syntax"));
+		Info.ExpireDuration = 5.0f;
+		FSlateNotificationManager::Get().AddNotification(Info);
+
+		return FReply::Handled();
+	}
+
+	FString WinSDKVersion;
+	if (!FPlatformMisc::QueryRegKey(HKEY_CURRENT_USER, *WinSDKSubKey, TEXT("ProductVersion"), WinSDKVersion))
+	{
+		if (!FPlatformMisc::QueryRegKey(HKEY_LOCAL_MACHINE, *WinSDKSubKey, TEXT("ProductVersion"), WinSDKVersion))
+		{
+			// Carry on for now, maybe it's in an unversioned folder
+		}
+	}
+	WinSDKVersion += TEXT(".0");
+
+	FString ArchFragment = PLATFORM_64BITS ? TEXT("x64") : TEXT("x86");
+	FString MakeCert = WinSDKFolder / TEXT("bin") / WinSDKVersion / ArchFragment / TEXT("makecert.exe");
+	if (!IPlatformFile::GetPlatformPhysical().FileExists(*MakeCert))
+	{
+		MakeCert = WinSDKFolder / TEXT("bin") / ArchFragment / TEXT("makecert.exe");
+		if (!IPlatformFile::GetPlatformPhysical().FileExists(*MakeCert))
+		{
+			FNotificationInfo Info(LOCTEXT("SigningCertificateFailed_MissingMakecert", "Failed to generate certificate: could not locate makecert.exe"));
+			Info.ExpireDuration = 3.0f;
+			FSlateNotificationManager::Get().AddNotification(Info);
+
+			return FReply::Handled();
+		}
+	}
+
+	FString Pvk2Pfx = WinSDKFolder / TEXT("bin") / WinSDKVersion / ArchFragment / TEXT("pvk2pfx.exe");
+	if (!IPlatformFile::GetPlatformPhysical().FileExists(*Pvk2Pfx))
+	{
+		Pvk2Pfx = WinSDKFolder / TEXT("bin") / ArchFragment / TEXT("pvk2pfx.exe");
+		if (!IPlatformFile::GetPlatformPhysical().FileExists(*Pvk2Pfx))
+		{
+			FNotificationInfo Info(LOCTEXT("SigningCertificateFailed_MissingPvk2pfx", "Failed to generate certificate: could not locate pvk2pfx.exe"));
+			Info.ExpireDuration = 3.0f;
+			FSlateNotificationManager::Get().AddNotification(Info);
+
+			return FReply::Handled();
+		}
+	}
+
+	FString TempCerFile = FPaths::CreateTempFilename(FPlatformProcess::UserTempDir(), TEXT("SigningCertificate-"), TEXT(".cer"));
+	FString TempPvkFile = FPaths::ChangeExtension(TempCerFile, TEXT(".pvk"));
+	FString TempPfxFile = FPaths::ChangeExtension(TempCerFile, TEXT(".pfx"));
+
+	void* ReadPipe = nullptr;
+	void* WritePipe = nullptr;
+	FPlatformProcess::CreatePipe(ReadPipe, WritePipe);
+	int32 ExitCode = 0xffffffff;
+	uint32 ProcId = 0;
+
+	FString MakeCertParams = FString::Printf(TEXT("-r -h 0 -n \"%s\" -eku %hs -pe -sv \"%s\" \"%s\""), *GetPublisherIdentityName(), szOID_PKIX_KP_CODE_SIGNING, *TempPvkFile, *TempCerFile);
+
+	// ExecProcess is simpler, but unfortunately forces any created window to be minimized.
+	// MakeCert pops UI to ask what password the new cert should have (answer is always none for this usage, but no way to say this on the command line)
+	FProcHandle MakeCertProc = FPlatformProcess::CreateProc(*MakeCert, *MakeCertParams, false, false, false, &ProcId, 0, nullptr, WritePipe);
+	FString MakeCertOutput;
+	FString LatestOutput = FPlatformProcess::ReadPipe(ReadPipe);
+	while (FPlatformProcess::IsProcRunning(MakeCertProc) || !LatestOutput.IsEmpty())
+	{
+		MakeCertOutput += LatestOutput;
+		LatestOutput = FPlatformProcess::ReadPipe(ReadPipe);
+		FPlatformProcess::Sleep(0);
+	}
+	bool GotReturnCode = FPlatformProcess::GetProcReturnCode(MakeCertProc, &ExitCode);
+	FPlatformProcess::CloseProc(MakeCertProc);
+	FPlatformProcess::ClosePipe(ReadPipe, WritePipe);
+	if (!GotReturnCode || ExitCode != 0)
+	{
+		FText NotificationMessage = FText::FormatOrdered(LOCTEXT("SigningCertificateFailed_MakecertFailed", "Failed to generate certificate: makecert.exe encountered an error: {0}"), FText::FromString(MakeCertOutput));
+		FNotificationInfo Info(NotificationMessage);
+		Info.ExpireDuration = 6.0f;
+		FSlateNotificationManager::Get().AddNotification(Info);
+
+		return FReply::Handled();
+	}
+
+	FString StdOutput;
+	FString StdError;
+	ExitCode = 0xffffffff;
+	FString Pvk2PfxParams = FString::Printf(TEXT("-pvk \"%s\" -spc \"%s\" -pfx \"%s\""), *TempPvkFile, *TempCerFile, *TempPfxFile);
+	GotReturnCode = FPlatformProcess::ExecProcess(*Pvk2Pfx, *Pvk2PfxParams, &ExitCode, &StdOutput, &StdError);
+	
+	if (!GotReturnCode || ExitCode != 0)
+	{
+		FText NotificationMessage = FText::FormatOrdered(LOCTEXT("SigningCertificateFailed_Pvk2PfxFailed", "Failed to generate certificate: pvk2pfx.exe encountered an error: {0}"), FText::FromString(StdError));
+		FNotificationInfo Info(NotificationMessage);
+		Info.ExpireDuration = 6.0f;
+		FSlateNotificationManager::Get().AddNotification(Info);
+
+		return FReply::Handled();
+	}
+
+	OnCertificatePicked(TempPfxFile);
+
+	IPlatformFile& PlatformFile = FPlatformFileManager::Get().GetPlatformFile();
+	PlatformFile.DeleteFile(*TempCerFile);
+	PlatformFile.DeleteFile(*TempPvkFile);
+	PlatformFile.DeleteFile(*TempPfxFile);
+
+	return FReply::Handled();
+}
+
+FString FHoloLensTargetSettingsCustomization::GetPublisherIdentityName() const
+{
+	FString PublisherIdentityName = GetDefault<UGeneralProjectSettings>()->CompanyDistinguishedName;
+	if (PublisherIdentityName.IsEmpty())
+	{
+		PublisherIdentityName = TEXT("CN=NoPublisher");
+	}
+	return PublisherIdentityName;
+}
+
+void FHoloLensTargetSettingsCustomization::LoadAndValidateSigningCertificate()
+{
+	SigningCertificateError->SetError(FText::GetEmpty());
+	SigningCertificateSubjectName = TEXT("Invalid Certificate");
+
+	FString CertificatePath = FPaths::ProjectDir() / TEXT("Build") / TEXT("HoloLens") / TEXT("SigningCertificate.pfx");
+	TArray<uint8> CertBytes;
+	if (FFileHelper::LoadFileToArray(CertBytes, *CertificatePath, FILEREAD_Silent))
+	{
+		CRYPT_DATA_BLOB CertBlob;
+		CertBlob.cbData = CertBytes.Num();
+		CertBlob.pbData = CertBytes.GetData();
+		HCERTSTORE CertStore = PFXImportCertStore(&CertBlob, nullptr, 0);
+		if (CertStore != nullptr)
+		{
+			PCCERT_CONTEXT SigningCertificateContext = CertEnumCertificatesInStore(CertStore, nullptr);
+			if (SigningCertificateContext)
+			{
+				DWORD StrType = CERT_X500_NAME_STR;
+				DWORD NumCharacters = CertGetNameString(SigningCertificateContext, CERT_NAME_RDN_TYPE, 0, &StrType, nullptr, 0);
+				SigningCertificateSubjectName.GetCharArray().SetNumZeroed(NumCharacters);
+				CertGetNameString(SigningCertificateContext, CERT_NAME_RDN_TYPE, 0, &StrType, SigningCertificateSubjectName.GetCharArray().GetData(), SigningCertificateSubjectName.GetAllocatedSize());
+
+				// The name string is a list of key/value pairs, their order may not match
+				bool bNameMatches = true;
+				{
+					TArray<FString> IdentityNameArray;
+					FString IdentityNameString = GetPublisherIdentityName();
+					IdentityNameString.ParseIntoArray(IdentityNameArray, TEXT(", "), true);
+
+					TArray<FString> SubjectNameArray;
+					SigningCertificateSubjectName.ParseIntoArray(SubjectNameArray, TEXT(", "), true);
+
+					if (IdentityNameArray.Num() != SubjectNameArray.Num())
+					{
+						bNameMatches = false;
+					}
+					else
+					{
+						IdentityNameArray.Sort();
+						SubjectNameArray.Sort();
+						const int Num = IdentityNameArray.Num();
+						for (int i = 0; i < Num; i++)
+						{
+							if (IdentityNameArray[i] != SubjectNameArray[i])
+							{
+								bNameMatches = false;
+								break;
+							}
+						}
+					}
+				}
+
+				if (!bNameMatches)
+				{
+					FText ErrorMessage = FText::FormatOrdered(LOCTEXT("CertificateInvalidSubjectName", "Certificate subject name[{0}] does not match Package / Identity / Publisher[{1}] in AppxManifest"), FText::FromString(SigningCertificateSubjectName), FText::FromString(GetPublisherIdentityName()));
+					SigningCertificateError->SetError(ErrorMessage);
+				}
+
+				// TODO - check private key, expiration, others?
+				CertFreeCertificateContext(SigningCertificateContext);
+			}
+			else
+			{
+				SigningCertificateError->SetError(LOCTEXT("CertificateEmptyStore", "No certificate found in SigningCertificate.pfx file"));
+			}
+			CertCloseStore(CertStore, 0);
+		}
+		else
+		{
+			DWORD FailureReason = GetLastError();
+			switch (FailureReason)
+			{
+			case ERROR_INVALID_PASSWORD:
+				SigningCertificateError->SetError(LOCTEXT("CertificateFilePasswordProtected", "SigningCertificate.pfx file is password protected"));
+				break;
+
+			default:
+				SigningCertificateError->SetError(LOCTEXT("CertificateUnknownError", "Unknown error loading SigningCertificate.pfx"));
+				break;
+			}
+		}
+	}
+	else
+	{
+		SigningCertificateError->SetError(LOCTEXT("CertificateMissing", "SigningCertificate.pfx does not exist"));
+	}
+}
+
+FString FHoloLensTargetSettingsCustomization::GetSigningCertificateSubjectName() const
+{
+	return SigningCertificateSubjectName;
+}
+
+void FHoloLensTargetSettingsCustomization::OnAutoDetectWin10SDKChanged(ECheckBoxState NewState, TSharedRef<IPropertyHandle> Win10SDKVersionPropertyHandle)
+{
+	Win10SDKVersionPropertyHandle->NotifyPreChange();
+	if (NewState == ECheckBoxState::Checked)
+	{
+		Win10SDKVersionPropertyHandle->SetValue(FString());
+		WindowsSDKSelector->SetSelectedItem(MakeShared<FString>());
+		WindowsSDKSelector->SetEnabled(false);
+	}
+	else
+	{
+		if (PlatformVersionOptions.Num() > 0)
+		{
+			Win10SDKVersionPropertyHandle->SetValue(*PlatformVersionOptions[0]);
+			WindowsSDKSelector->SetSelectedItem(PlatformVersionOptions[0]);
+			WindowsSDKSelector->SetEnabled(true);
+		}
+	}
+	Win10SDKVersionPropertyHandle->NotifyPostChange(EPropertyChangeType::ValueSet);
+}
+
+ECheckBoxState FHoloLensTargetSettingsCustomization::IsAutoDetectWin10SDKChecked() const
+{
+	return GetDefault<UHoloLensTargetSettings>()->Windows10SDKVersion.IsEmpty() ? ECheckBoxState::Checked : ECheckBoxState::Unchecked;
+}
+
+#undef LOCTEXT_NAMESPACE
+
+#include "Windows/HideWindowsPlatformTypes.h"
diff --git a/Engine/Source/Developer/HoloLens/HoloLensPlatformEditor/Private/HoloLensTargetSettingsCustomization.h b/Engine/Source/Developer/HoloLens/HoloLensPlatformEditor/Private/HoloLensTargetSettingsCustomization.h
new file mode 100644
--- /dev/null
+++ b/Engine/Source/Developer/HoloLens/HoloLensPlatformEditor/Private/HoloLensTargetSettingsCustomization.h
@@ -0,0 +1,53 @@
+// Copyright Epic Games, Inc. All Rights Reserved.
+
+#pragma once
+
+#include "PropertyEditorModule.h"
+#include "IDetailCustomization.h"
+#include "DetailLayoutBuilder.h"
+#include "IDetailGroup.h"
+#include "TargetPlatformAudioCustomization.h"
+
+class SErrorHint;
+class STextComboBox;
+
+class FHoloLensTargetSettingsCustomization : public IDetailCustomization
+{
+public:
+
+	// Makes a new instance of this detail layout class for a specific detail view requesting it
+	static TSharedRef<IDetailCustomization> MakeInstance();
+
+	// IDetailCustomization interface
+	virtual void CustomizeDetails(IDetailLayoutBuilder& DetailLayout) override;
+	// End of IDetailCustomization interface
+private:
+
+	void AddWidgetForCapability(IDetailLayoutBuilder& DetailBuilder, TSharedRef<IPropertyHandle> CapabilityList, const FString& CapabilityName, const FText& CapabilityCaption, const FText& CapabilityTooltip, bool bForAdvanced);
+	void AddWidgetForPlatformVersion(IDetailLayoutBuilder& DetailBuilder, TSharedRef<IPropertyHandle> PropertyHandle, TSharedPtr<STextComboBox>* OutVersionSelector = nullptr, bool bForAdvanced = false, const FString& SpecificVersionOverrideString = FString());
+	void AddWidgetForTargetDeviceFamily(IDetailLayoutBuilder& DetailBuilder, TSharedRef<IPropertyHandle> PropertyHandle);
+	static FString GetNameForSigningCertificate(const FString &CertificatePath);
+	void InitSupportedPlatformVersions();
+	void InitTargetDeviceFamilyOptions();
+	ECheckBoxState IsCapabilityChecked(TSharedRef<IPropertyHandle> CapabilityList, const FString CapabilityName) const;
+	void OnCapabilityStateChanged(ECheckBoxState CheckState, TSharedRef<IPropertyHandle> CapabilityList, const FString CapabilityName);
+	void OnCertificatePicked(const FString& PickedPath);
+	void OnSelectedItemChanged(TSharedPtr<FString> NewValue, ESelectInfo::Type SelectInfo, TSharedRef<IPropertyHandle> HandlePtr);
+	void OnAutoDetectWin10SDKChanged(ECheckBoxState NewState, TSharedRef<IPropertyHandle> Win10SDKVersionPropertyHandle);
+	ECheckBoxState IsAutoDetectWin10SDKChecked() const;
+	EVisibility GetManualWin10SDKWidgetVisibility() const;
+
+	FString GetSigningCertificateSubjectName() const;
+	FReply GenerateSigningCertificate();
+	FString GetPublisherIdentityName() const;
+	void LoadAndValidateSigningCertificate();
+
+	TArray<TSharedPtr<FString>> PlatformVersionOptions;
+	TArray<TSharedPtr<FString>> TargetDeviceFamilyOptions;
+
+	FString SigningCertificateSubjectName;
+	TSharedPtr<SErrorHint> SigningCertificateError;
+	TSharedPtr<STextComboBox> WindowsSDKSelector;
+
+	FAudioPluginWidgetManager AudioPluginManager;
+};
\ No newline at end of file
diff --git a/Engine/Source/Developer/HoloLens/HoloLensPlatformEditor/Public/HoloLensPlatformEditor.h b/Engine/Source/Developer/HoloLens/HoloLensPlatformEditor/Public/HoloLensPlatformEditor.h
new file mode 100644
--- /dev/null
+++ b/Engine/Source/Developer/HoloLens/HoloLensPlatformEditor/Public/HoloLensPlatformEditor.h
@@ -0,0 +1,33 @@
+// Copyright Epic Games, Inc. All Rights Reserved.
+
+#pragma once
+
+#include "CoreMinimal.h"
+
+struct HOLOLENSPLATFORMEDITOR_API FHoloLensSDKVersion
+{
+	/** The whole version string for the SDK */
+	FString VersionString;
+	/** The parsed parts of the version string (1.2.3.4) */
+	int32 Version1;
+	int32 Version2;
+	int32 Version3;
+	int32 Version4;
+
+	FHoloLensSDKVersion()
+	{
+		Version1 = Version2 = Version3 = Version4 = 0;
+	}
+
+	FHoloLensSDKVersion(const FString& InVersionString, int32 InVer1, int32 InVer2, int32 InVer3, int32 InVer4)
+		: VersionString(InVersionString)
+		, Version1(InVer1)
+		, Version2(InVer2)
+		, Version3(InVer3)
+		, Version4(InVer4)
+	{
+	}
+
+	/** Returns the list of SDK versions for this machine */
+	static TArray<FHoloLensSDKVersion> GetSDKVersions();
+};
diff --git a/Engine/Source/Developer/HoloLens/HoloLensTargetPlatform/HoloLensTargetPlatform.Build.cs b/Engine/Source/Developer/HoloLens/HoloLensTargetPlatform/HoloLensTargetPlatform.Build.cs
new file mode 100644
--- /dev/null
+++ b/Engine/Source/Developer/HoloLens/HoloLensTargetPlatform/HoloLensTargetPlatform.Build.cs
@@ -0,0 +1,54 @@
+// Copyright Epic Games, Inc. All Rights Reserved.
+
+using System;
+using System.IO;
+using UnrealBuildTool;
+
+public class HoloLensTargetPlatform : ModuleRules
+{
+	public HoloLensTargetPlatform(ReadOnlyTargetRules Target) : base(Target)
+    {
+		PrivateDependencyModuleNames.AddRange(
+			new string[] {
+				"Core",
+                "CoreUObject",
+				"Settings",
+                "EngineSettings",
+                "TargetPlatform",
+				"DesktopPlatform",
+				"HTTP",
+                "Json",
+                "HoloLensPlatformEditor",
+            }
+        );
+
+		PrivateIncludePathModuleNames.Add("Settings");
+
+        int Win10Build = 0;
+        Version ver = null;
+        if (Version.TryParse(Target.WindowsPlatform.WindowsSdkVersion, out ver))
+        {
+            Win10Build = ver.Build;
+        }
+
+        //there is a WinSDK bug that prevented to include the file into this version
+		if(Win10Build != 0 && Win10Build != 16299)
+        {
+            PrivateDefinitions.Add("APPXPACKAGING_ENABLE=1");
+        }
+		else
+        {
+            PrivateDefinitions.Add("APPXPACKAGING_ENABLE=0");
+        }
+
+        // compile withEngine
+        if (Target.bCompileAgainstEngine)
+		{
+			PrivateDependencyModuleNames.Add("Engine");
+			PrivateIncludePathModuleNames.Add("TextureCompressor");
+		}
+
+		PublicSystemLibraries.Add("shlwapi.lib");
+        PublicSystemLibraries.Add("runtimeobject.lib");
+	}
+}
diff --git a/Engine/Source/Developer/HoloLens/HoloLensTargetPlatform/Private/HoloLensDeviceDetector.cpp b/Engine/Source/Developer/HoloLens/HoloLensTargetPlatform/Private/HoloLensDeviceDetector.cpp
new file mode 100644
--- /dev/null
+++ b/Engine/Source/Developer/HoloLens/HoloLensTargetPlatform/Private/HoloLensDeviceDetector.cpp
@@ -0,0 +1,586 @@
+// Copyright Epic Games, Inc. All Rights Reserved.
+
+#include "IHoloLensDeviceDetector.h"
+#include "HttpModule.h"
+#include "PlatformHttp.h"
+#include "Interfaces/IHttpResponse.h"
+#include "Misc/ScopeLock.h"
+#include "Misc/Base64.h"
+#include "Containers/UnrealString.h"
+#include "Dom/JsonObject.h"
+#include "HoloLensTargetPlatform.h"
+
+#pragma warning(push)
+#pragma warning(disable:4265 4459 6101)
+
+#include "Serialization/JsonSerializer.h"
+
+#include "Windows/WindowsHWrapper.h"
+#include "Windows/AllowWindowsPlatformTypes.h"
+#include "Windows/AllowWindowsPlatformAtomics.h"
+
+#define sealed final 
+
+#include <windows.h>
+#include <wrl/client.h>
+#include <wrl/event.h>
+#include <wrl/implements.h>
+#include <wrl/wrappers/corewrappers.h>
+#include <Windows.Devices.Enumeration.h>
+#include <Windows.foundation.h>
+#include <Windows.foundation.collections.h>
+#undef sealed  
+#pragma warning(pop)
+
+#include <vector>
+#include <memory>
+
+using namespace ABI::Windows::Foundation;
+using namespace ABI::Windows::Foundation::Collections;
+using namespace ABI::Windows::Devices::Enumeration;
+using namespace Microsoft::WRL;
+using namespace Microsoft::WRL::Wrappers;
+
+namespace 
+{
+	typedef __FITypedEventHandler_2_Windows__CDevices__CEnumeration__CDeviceWatcher_Windows__CDevices__CEnumeration__CDeviceInformation AddedHandler;
+	typedef __FITypedEventHandler_2_Windows__CDevices__CEnumeration__CDeviceWatcher_Windows__CDevices__CEnumeration__CDeviceInformationUpdate UpdateHandler;
+	typedef __FITypedEventHandler_2_Windows__CDevices__CEnumeration__CDeviceWatcher_IInspectable EnumerationCompletedHandler;
+	typedef __FITypedEventHandler_2_Windows__CDevices__CEnumeration__CDeviceWatcher_IInspectable StoppedHandler;
+
+	class StringCollectionIterator : public RuntimeClass<IIterator<HSTRING>>
+	{
+		InspectableClass(L"Library1.StringCollectionIterator", BaseTrust)
+	public:
+
+		virtual /* propget */ HRESULT STDMETHODCALLTYPE get_Current(_Out_ HSTRING *current)
+		{
+			if ((CurrentIdx >= InternalVector->size()))
+			{
+				return E_INVALIDARG;
+			}
+
+			*current = (*InternalVector)[CurrentIdx];
+			return S_OK;
+		}
+
+		virtual /* propget */ HRESULT STDMETHODCALLTYPE get_HasCurrent(_Out_ boolean *hasCurrent)
+		{
+			*hasCurrent = (CurrentIdx < InternalVector->size());
+			return S_OK;
+		}
+
+		virtual HRESULT STDMETHODCALLTYPE MoveNext(_Out_ boolean *hasCurrent)
+		{
+			*hasCurrent = (++CurrentIdx < InternalVector->size());
+			return S_OK;
+		}
+
+		std::shared_ptr<std::vector<HSTRING>> InternalVector;
+		size_t CurrentIdx = 0;
+	};
+
+	class StringCollection : public RuntimeClass<IIterable<HSTRING>>
+	{
+		InspectableClass(L"Library1.StringCollection", BaseTrust)
+	public:
+
+		virtual HRESULT STDMETHODCALLTYPE First(_Outptr_result_maybenull_ IIterator<HSTRING> **first)
+		{
+			ComPtr<StringCollectionIterator> iter = Make<StringCollectionIterator>();
+			iter->InternalVector = InternalVector;
+			return iter.CopyTo(first);
+		}
+
+		std::shared_ptr<std::vector<HSTRING>> InternalVector = std::make_shared<std::vector<HSTRING>>();
+	};
+
+}
+
+namespace HoloLensDeviceTypes
+{
+	const FName HoloLens = "Windows.HoloLens";
+	const FName HoloLensEmulation = "Windows.HoloLensEmulation";
+}
+
+class FHoloLensDeviceDetector : public IHoloLensDeviceDetector
+{
+public:
+	FHoloLensDeviceDetector();
+	virtual ~FHoloLensDeviceDetector();
+
+	virtual void StartDeviceDetection();
+	virtual void StopDeviceDetection();
+	virtual void TryAddDevice(const FString& DeviceId, const FString& DeviceUserFriendlyName, const FString& Username, const FString& Password);
+
+	virtual FOnDeviceDetected& OnDeviceDetected()			{ return DeviceDetected; }
+	virtual const TArray<FHoloLensDeviceInfo> GetKnownDevices()	{ return KnownDevices; }
+
+private:
+
+	void OnDeviceWatcherDeviceAdded(IDeviceInformation* Info);
+
+	void AddDevice(const FHoloLensDeviceInfo& Info);
+
+
+private:
+
+	ComPtr < IDeviceWatcher > DeviceWatcher;
+	EventRegistrationToken addedToken;
+	EventRegistrationToken updatedToken;
+	EventRegistrationToken removedToken;
+
+	FOnDeviceDetected DeviceDetected;
+	FCriticalSection DevicesLock;
+	TArray<FHoloLensDeviceInfo> KnownDevices;
+};
+
+
+IHoloLensDeviceDetectorPtr IHoloLensDeviceDetector::Create()
+{
+	return TSharedPtr<IHoloLensDeviceDetector, ESPMode::ThreadSafe>(new FHoloLensDeviceDetector);
+}
+
+
+
+FHoloLensDeviceDetector::FHoloLensDeviceDetector()
+{
+}
+
+FHoloLensDeviceDetector::~FHoloLensDeviceDetector()
+{
+	StopDeviceDetection();
+}
+
+void FHoloLensDeviceDetector::StartDeviceDetection()
+{
+	if (DeviceWatcher)
+	{
+		return;
+	}
+
+#if APPXPACKAGING_ENABLE
+
+	HRESULT hr;
+
+	hr = ::RoInitialize(RO_INIT_MULTITHREADED);
+	if (FAILED(hr)) { return; }
+
+	ComPtr <IDeviceInformationStatics2> DeviceInformationStatics2;
+	hr = GetActivationFactory(HStringReference(RuntimeClass_Windows_Devices_Enumeration_DeviceInformation).Get(), &DeviceInformationStatics2);
+	if (FAILED(hr)) { return; }
+
+	ComPtr<StringCollection> AdditionalProperties = Make<StringCollection>();
+
+	HSTRING AqsFilter = HStringReference(TEXT(
+		"System.Devices.AepService.ProtocolId:={4526e8c1-8aac-4153-9b16-55e86ada0e54} AND "
+		"System.Devices.Dnssd.Domain:=\"local\" AND "
+		"System.Devices.Dnssd.ServiceName:=\"_wdp._tcp\" "
+		"")).Get();
+
+	AdditionalProperties->InternalVector->push_back(HStringReference(L"System.Devices.Dnssd.HostName").Get());
+	AdditionalProperties->InternalVector->push_back(HStringReference(L"System.Devices.Dnssd.ServiceName").Get());
+	AdditionalProperties->InternalVector->push_back(HStringReference(L"System.Devices.Dnssd.PortNumber").Get());
+	AdditionalProperties->InternalVector->push_back(HStringReference(L"System.Devices.Dnssd.TextAttributes").Get());
+	AdditionalProperties->InternalVector->push_back(HStringReference(L"System.Devices.IpAddress").Get());
+
+	hr = DeviceInformationStatics2->CreateWatcherWithKindAqsFilterAndAdditionalProperties(
+		AqsFilter,
+		AdditionalProperties.Get(),
+		DeviceInformationKind::DeviceInformationKind_AssociationEndpointService,
+		&DeviceWatcher);
+	if (FAILED(hr)) { return; }
+
+	hr = DeviceWatcher->add_Added(Callback<AddedHandler>([this](IDeviceWatcher* watcher, IDeviceInformation* info) -> HRESULT
+	{
+		//devices added
+		OnDeviceWatcherDeviceAdded(info);
+		return S_OK;
+	}).Get(), &addedToken);
+
+	//these callbacks should be kept for continious monitoring for new devices
+	hr = DeviceWatcher->add_Updated(Callback<UpdateHandler>([](IDeviceWatcher* watcher, IDeviceInformationUpdate* info) -> HRESULT
+	{
+		return S_OK;
+	}).Get(), &updatedToken);
+
+
+	hr = DeviceWatcher->add_Removed(Callback<UpdateHandler>([](IDeviceWatcher* watcher, IDeviceInformationUpdate* info) -> HRESULT
+	{
+		return S_OK;
+	}).Get(), &removedToken);
+
+
+	DeviceWatcher->Start();
+#endif
+}
+
+void FHoloLensDeviceDetector::StopDeviceDetection()
+{
+	if (DeviceWatcher != nullptr)
+	{
+		DeviceWatcher->remove_Added(addedToken);
+		DeviceWatcher->remove_Updated(updatedToken);
+		DeviceWatcher->remove_Removed(removedToken);
+
+		DeviceWatcher->Stop();
+		DeviceWatcher = nullptr;
+	}
+}
+
+
+void FHoloLensDeviceDetector::OnDeviceWatcherDeviceAdded(IDeviceInformation* Info)
+{
+	SslCertDisabler disabler;
+	FHoloLensDeviceInfo NewDevice;
+	uint32 WdpPort;
+	FString DeviceIp;
+	FString Protocol = TEXT("http");
+	FName DeviceType = NAME_None;
+
+	NewDevice.IsLocal = 0;
+
+	HString hstr;
+	HRESULT hr;
+
+	hr = Info->get_Id(hstr.ReleaseAndGetAddressOf());
+	if (FAILED(hr)) { return; }
+
+	NewDevice.WindowsDeviceId = hstr.GetRawBuffer(nullptr);
+
+	{
+		ComPtr<IMapView<HSTRING, IInspectable*>> Properties;
+		ComPtr<IInspectable> Inspectable;
+		ComPtr<IPropertyValue> PropertyValue;
+
+		UINT32 uiValue;
+		HSTRING * pArrayString;
+
+		hr = Info->get_Properties(&Properties);
+		if (FAILED(hr)) { return; }
+
+		//Hostname
+		hr = Properties->Lookup(HStringReference(L"System.Devices.Dnssd.HostName").Get(), &Inspectable);
+		if (FAILED(hr)) { return; }
+
+		hr = Inspectable.As(&PropertyValue);
+		if (FAILED(hr)) { return; }
+
+		hr = PropertyValue->GetString(hstr.ReleaseAndGetAddressOf());
+		if (FAILED(hr)) { return; }
+
+		NewDevice.HostName = hstr.GetRawBuffer(nullptr);
+		NewDevice.HostName.RemoveFromEnd(TEXT(".local"));
+
+		//Port
+		hr = Properties->Lookup(HStringReference(L"System.Devices.Dnssd.PortNumber").Get(), &Inspectable);
+		if (FAILED(hr)) { return; }
+
+		hr = Inspectable.As(&PropertyValue);
+		if (FAILED(hr)) { return; }
+
+		hr = PropertyValue->GetUInt32(&uiValue);
+		if (FAILED(hr)) { return; }
+
+		WdpPort = uiValue;
+
+		//Attributes
+		hr = Properties->Lookup(HStringReference(L"System.Devices.Dnssd.TextAttributes").Get(), &Inspectable);
+		if (FAILED(hr)) { return; }
+
+		hr = Inspectable.As(&PropertyValue);
+		if (FAILED(hr)) { return; }
+
+		hr = PropertyValue->GetStringArray(&uiValue, &pArrayString);
+		if (FAILED(hr)) { return; }
+
+		for (UINT32 i = 0; i < uiValue; ++i)
+		{
+			*hstr.ReleaseAndGetAddressOf() = pArrayString[i];
+			const TCHAR* attrValue = hstr.GetRawBuffer(nullptr);
+
+			//parse hstr
+			FString Architecture;
+			uint32 SecurePort;
+			if (FParse::Value(attrValue, TEXT("S="), SecurePort))
+			{
+				Protocol = TEXT("https");
+				WdpPort = SecurePort;
+			}
+			else if (FParse::Value(attrValue, TEXT("D="), NewDevice.DeviceTypeName))
+			{
+
+			}
+			else if (FParse::Value(attrValue, TEXT("A="), Architecture))
+			{
+				if (Architecture == TEXT("AMD64"))
+				{
+					NewDevice.Architecture = EHoloLensArchitecture::X64;
+				}
+				else if (Architecture == TEXT("x86"))
+				{
+					NewDevice.Architecture = EHoloLensArchitecture::X86;
+				}
+				else if (Architecture == TEXT("ARM64"))
+				{
+					NewDevice.Architecture = EHoloLensArchitecture::ARM64;
+				}
+				else if (Architecture == TEXT("ARM"))
+				{
+					NewDevice.Architecture = EHoloLensArchitecture::ARM32;
+				}
+				else
+				{
+					//defaulting
+					NewDevice.Architecture = EHoloLensArchitecture::X86;
+				}
+			}
+		}
+
+		//IP Address
+		hr = Properties->Lookup(HStringReference(L"System.Devices.IpAddress").Get(), &Inspectable);
+		if (FAILED(hr)) { return; }
+
+		hr = Inspectable.As(&PropertyValue);
+		if (FAILED(hr)) { return; }
+
+		hr = PropertyValue->GetStringArray(&uiValue, &pArrayString);
+		if (FAILED(hr) || pArrayString == nullptr) { return; }
+
+		if (uiValue == 0) { return; }
+
+		*hstr.ReleaseAndGetAddressOf() = pArrayString[0];
+
+		DeviceIp = hstr.GetRawBuffer(nullptr);
+	}
+
+	NewDevice.WdpUrl = FString::Printf(TEXT("%s://%s:%d/"), *Protocol, *DeviceIp, WdpPort);
+
+
+	// Now make a test request against the device to determine whether or not the Device Portal requires authentication.
+	// If it does, the user will have to add it manually so we can collect the username and password.
+	TSharedRef<IHttpRequest, ESPMode::ThreadSafe> TestRequest = FHttpModule::Get().CreateRequest();
+	TestRequest->SetVerb(TEXT("GET"));
+	TestRequest->SetURL(NewDevice.WdpUrl);
+	TestRequest->OnProcessRequestComplete().BindLambda(
+		[=](FHttpRequestPtr HttpRequest, FHttpResponsePtr HttpResponse, bool bSucceeded)
+	{
+		bool RequiresCredentials = false;
+		bool ReachedDevicePortal = false;
+		if (bSucceeded && HttpResponse.IsValid())
+		{
+			int32 ResponseCode = HttpResponse->GetResponseCode();
+			RequiresCredentials = ResponseCode == EHttpResponseCodes::Denied;
+			ReachedDevicePortal = RequiresCredentials || EHttpResponseCodes::IsOk(ResponseCode);
+		}
+
+		if (ReachedDevicePortal || NewDevice.IsLocal)
+		{
+			FHoloLensDeviceInfo ValidatedDevice = NewDevice;
+			ValidatedDevice.RequiresCredentials = RequiresCredentials;
+			AddDevice(ValidatedDevice);
+		}
+	});
+	TestRequest->ProcessRequest();
+}
+
+
+
+bool GetJsonField(FString& OutVal, const TSharedPtr<FJsonObject>& JsonObject, const TCHAR* FieldName)
+{
+	OutVal = TEXT("");
+
+	auto JsonValue = JsonObject->Values.Find(FieldName);
+
+	if (!JsonValue)
+	{
+		return false;
+	}
+
+	if ((*JsonValue)->Type != EJson::String)
+	{
+		return false;
+	}
+
+	(*JsonValue)->AsArgumentType(OutVal);
+	return true;
+}
+
+
+bool GetJsonField(int64& OutVal, const TSharedPtr<FJsonObject>& JsonObject, const TCHAR* FieldName)
+{
+	OutVal = 0;
+	double Value;
+
+	auto JsonValue = JsonObject->Values.Find(FieldName);
+
+	if (!JsonValue)
+	{
+		return false;
+	}
+
+	if ((*JsonValue)->Type != EJson::Number)
+	{
+		return false;
+	}
+
+	(*JsonValue)->AsArgumentType(Value);
+	OutVal = (int32)Value;
+	return true;
+}
+
+
+
+void FHoloLensDeviceDetector::TryAddDevice(const FString& DeviceId, const FString& DeviceUserFriendlyName, const FString& Username, const FString& Password)
+{
+	SslCertDisabler disabler;
+	
+	auto DeviceInfo = MakeShared<FHoloLensDeviceInfo, ESPMode::ThreadSafe>();
+
+	DeviceInfo->HostName = DeviceUserFriendlyName;
+	DeviceInfo->WdpUrl = DeviceId;
+	DeviceInfo->IsLocal = 0;
+	DeviceInfo->Username = Username;
+	DeviceInfo->Password = Password;
+	DeviceInfo->RequiresCredentials = Username.IsEmpty() && Password.IsEmpty() ? 0 : 1;
+	DeviceInfo->Architecture = EHoloLensArchitecture::Unknown;
+	if (!DeviceInfo->WdpUrl.Contains("://"))
+	{
+		DeviceInfo->WdpUrl = "http://" + DeviceInfo->WdpUrl;
+	}
+	if (!DeviceInfo->WdpUrl.EndsWith(TEXT("/")))
+	{
+		DeviceInfo->WdpUrl += TEXT("/");
+	}
+
+	TSharedRef<IHttpRequest, ESPMode::ThreadSafe> TestRequest = FHttpModule::Get().CreateRequest();
+	TestRequest->SetVerb(TEXT("GET"));
+	TestRequest->SetURL(DeviceInfo->WdpUrl + TEXT("api/os/info"));
+
+	if (DeviceInfo->RequiresCredentials)
+	{
+		FString AuthentificationString = FString::Format(TEXT("auto-{0}:{1}"), { DeviceInfo->Username, DeviceInfo->Password });
+		FString AuthHeaderString = FString::Format(TEXT("Basic {0}"), { FBase64::Encode(AuthentificationString) });
+		TestRequest->AppendToHeader(TEXT("Authorization"), AuthHeaderString);
+	}
+
+	UE_LOG(LogHoloLensTargetPlatform, Log, TEXT("The device %s is being added manually"), *DeviceInfo->WdpUrl);
+
+	TestRequest->OnProcessRequestComplete().BindLambda(
+		[=](FHttpRequestPtr HttpRequest, FHttpResponsePtr HttpResponse, bool bSucceeded)
+	{
+		if (!bSucceeded || !HttpResponse.IsValid())
+		{
+			UE_LOG(LogHoloLensTargetPlatform, Error, TEXT("The device %s doesn't respond, the adding fails"), *DeviceInfo->WdpUrl);
+			return;
+		}
+
+		if (HttpResponse->GetResponseCode() != 200) //we can accept just 200 OK here
+		{
+			UE_LOG(LogHoloLensTargetPlatform, Error, TEXT("The device %s responds with the http error %d, the adding fails"), *DeviceInfo->WdpUrl, HttpResponse->GetResponseCode());
+			return;
+		}
+
+		FString Json = HttpResponse->GetContentAsString();
+		auto RootObject = TSharedPtr<FJsonObject>();
+		auto Reader = TJsonReaderFactory<TCHAR>::Create(Json);
+
+		if (!FJsonSerializer::Deserialize(Reader, RootObject))
+		{
+			UE_LOG(LogHoloLensTargetPlatform, Error, TEXT("The device %s responds but the answer couldn't be decoded, the adding fails"), *DeviceInfo->WdpUrl);
+			return;
+		}
+
+		FString str;
+		GetJsonField(str, RootObject, TEXT("OsEdition"));
+		if (str != TEXT("Holographic"))
+		{
+			UE_LOG(LogHoloLensTargetPlatform, Error, TEXT("The device %s responds but its OS isn't Holographic, the adding fails"), *DeviceInfo->WdpUrl);
+			return;
+		}
+
+		GetJsonField(str, RootObject, TEXT("Platform"));
+		if (str.StartsWith(TEXT("HoloLens")))
+		{
+			DeviceInfo->DeviceTypeName = HoloLensDeviceTypes::HoloLens;
+		} 
+		else if (str == TEXT("Virtual Machine"))
+		{
+			DeviceInfo->DeviceTypeName = HoloLensDeviceTypes::HoloLensEmulation;
+		}
+		else
+		{
+			UE_LOG(LogHoloLensTargetPlatform, Error, TEXT("The device %s responds but its platform isn't HoloLens/Emulator, the adding fails"), *DeviceInfo->WdpUrl);
+			return;
+		}
+
+		GetJsonField(str, RootObject, TEXT("OsVersion"));
+		{
+			auto ptr = GetData(str);
+			int idx = 0;
+			while (*ptr)
+			{
+				if (*ptr++ == '.')
+				{
+					if (++idx == 2)
+					{
+						break;
+					}
+				}
+			}
+			if (*ptr)
+			{
+				if (_wcsicmp(ptr, TEXT("amd64")) == 0)
+				{
+					DeviceInfo->Architecture = EHoloLensArchitecture::X64;
+				}
+				else if (_wcsicmp(ptr, TEXT("x86")) == 0)
+				{
+					DeviceInfo->Architecture = EHoloLensArchitecture::X86;
+				}
+				else if (_wcsicmp(ptr, TEXT("arm64")) == 0)
+				{
+					DeviceInfo->Architecture = EHoloLensArchitecture::ARM64;
+				}
+			}
+		}
+		GetJsonField(DeviceInfo->HostName, RootObject, TEXT("ComputerName"));
+
+		AddDevice(*DeviceInfo);
+	});
+	TestRequest->ProcessRequest();
+
+}
+
+
+
+void FHoloLensDeviceDetector::AddDevice(const FHoloLensDeviceInfo& Info)
+{
+	FScopeLock Lock(&DevicesLock);
+	for (const FHoloLensDeviceInfo& Existing : KnownDevices)
+	{
+		if (Existing.HostName == Info.HostName)
+		{
+			UE_LOG(LogHoloLensTargetPlatform, Log, TEXT("The device %s already exists"), *Info.WdpUrl);
+			return;
+		}
+	}
+
+	UE_LOG(LogHoloLensTargetPlatform, Display, TEXT("The device %s have been successefully added"), *Info.WdpUrl);
+
+	KnownDevices.Add(Info);
+	DeviceDetected.Broadcast(Info);
+}
+
+SslCertDisabler::SslCertDisabler()
+{
+	prevValue = FWindowsPlatformHttp::VerifyPeerSslCertificate(false);
+}
+
+SslCertDisabler::~SslCertDisabler()
+{
+	FWindowsPlatformHttp::VerifyPeerSslCertificate(prevValue);
+}
+
+#include "Windows/HideWindowsPlatformAtomics.h"
+#include "Windows/HideWindowsPlatformTypes.h"
diff --git a/Engine/Source/Developer/HoloLens/HoloLensTargetPlatform/Private/HoloLensTargetDevice.cpp b/Engine/Source/Developer/HoloLens/HoloLensTargetPlatform/Private/HoloLensTargetDevice.cpp
new file mode 100644
--- /dev/null
+++ b/Engine/Source/Developer/HoloLens/HoloLensTargetPlatform/Private/HoloLensTargetDevice.cpp
@@ -0,0 +1,497 @@
+// Copyright Epic Games, Inc. All Rights Reserved.
+
+#include "HoloLensTargetDevice.h"
+
+#include "Misc/Paths.h"
+
+#if PLATFORM_WINDOWS || PLATFORM_HOLOLENS
+#include "Windows/AllowWindowsPlatformTypes.h"
+
+#endif
+#include "HttpModule.h"
+#include "PlatformHttp.h"
+#include "Interfaces/IHttpResponse.h"
+#include "Misc/ScopeLock.h"
+#include "Misc/Base64.h"
+#include "Containers/UnrealString.h"
+#include "Dom/JsonObject.h"
+#include "Serialization/JsonSerializer.h"
+#include "Misc/ScopeLock.h"
+
+#pragma warning(push)
+#pragma warning(disable:4265 4459)
+
+#include <wrl/client.h>
+#include <wrl/wrappers/corewrappers.h>
+#pragma warning(pop)
+
+#include <shlwapi.h>
+
+#if (NTDDI_VERSION < NTDDI_WIN8)
+#pragma push_macro("NTDDI_VERSION")
+#undef NTDDI_VERSION
+#define NTDDI_VERSION NTDDI_WIN8
+//this header cannot be directly imported because of current _WIN32_WINNT less then 0x0602 (the value assigned in UEBuildWindows.cs:139)
+//the macro code added from couple interface declarations, it doesn't affect to any imported function
+#include <shobjidl.h> 
+#pragma pop_macro("NTDDI_VERSION")
+#else
+#include <shobjidl.h> 
+#endif
+
+#if APPXPACKAGING_ENABLE 
+#include <AppxPackaging.h>
+#endif
+
+#include "Microsoft/COMPointer.h"
+
+#if WITH_ENGINE
+#include "Engine/World.h"
+#include "TimerManager.h"
+#endif
+#include <functional>
+
+
+
+
+namespace WindowsMixedReality
+{
+	bool PackageProject(const wchar_t* StreamPath, bool PathIsActuallyPackage, const wchar_t* Params, UINT*& OutProcessId)
+	{
+#if APPXPACKAGING_ENABLE
+		using namespace Microsoft::WRL;
+		ComPtr<IAppxFactory> AppxFactory;
+		if (FAILED(CoCreateInstance(CLSID_AppxFactory, nullptr, CLSCTX_INPROC_SERVER, IID_PPV_ARGS(&AppxFactory))))
+		{
+			return false;
+		}
+
+		ComPtr<IStream> ReaderStream;
+		if (FAILED(SHCreateStreamOnFileEx(StreamPath, STGM_READ, 0, FALSE, nullptr, &ReaderStream)))
+		{
+			return false;
+		}
+
+		ComPtr<IAppxManifestReader> ManifestReader;
+		if (PathIsActuallyPackage)
+		{
+			ComPtr<IAppxPackageReader> PackageReader;
+			if (FAILED(AppxFactory->CreatePackageReader(ReaderStream.Get(), &PackageReader)))
+			{
+				return false;
+			}
+
+			if (FAILED(PackageReader->GetManifest(&ManifestReader)))
+			{
+				return false;
+			}
+		}
+		else
+		{
+			if (FAILED(AppxFactory->CreateManifestReader(ReaderStream.Get(), &ManifestReader)))
+			{
+				return false;
+			}
+		}
+
+		ComPtr<IAppxManifestApplicationsEnumerator> AppEnumerator;
+		if (FAILED(ManifestReader->GetApplications(&AppEnumerator)))
+		{
+			return false;
+		}
+
+		ComPtr<IAppxManifestApplication> ApplicationMetadata;
+		if (FAILED(AppEnumerator->GetCurrent(&ApplicationMetadata)))
+		{
+			return false;
+		}
+
+		ComPtr<IPackageDebugSettings> PackageDebugSettings;
+		if (SUCCEEDED(CoCreateInstance(CLSID_PackageDebugSettings, nullptr, CLSCTX_INPROC_SERVER, IID_PPV_ARGS(&PackageDebugSettings))))
+		{
+			ComPtr<IAppxManifestPackageId> ManifestPackageId;
+			if (SUCCEEDED(ManifestReader->GetPackageId(&ManifestPackageId)))
+			{
+				LPWSTR PackageFullName = nullptr;
+				// Avoid warning C6387 (static analysis assumes PackageFullName could be null even if the GetPackageFullName call succeeds.
+				if (SUCCEEDED(ManifestPackageId->GetPackageFullName(&PackageFullName)) && PackageFullName != nullptr)
+				{
+					PackageDebugSettings->EnableDebugging(PackageFullName, nullptr, nullptr);
+					CoTaskMemFree(PackageFullName);
+				}
+			}
+		}
+
+		LPWSTR Aumid = nullptr;
+		// Avoid warning C6387 (static analysis assumes Aumid could be null even if the GetAppUserModelId call succeeds.
+		if (FAILED(ApplicationMetadata->GetAppUserModelId(&Aumid)) || Aumid == nullptr)
+		{
+			return false;
+		}
+
+		bool ActivationSuccess = false;
+		ComPtr<IApplicationActivationManager> ActivationManager;
+		if (SUCCEEDED(CoCreateInstance(CLSID_ApplicationActivationManager, nullptr, CLSCTX_INPROC_SERVER, IID_PPV_ARGS(&ActivationManager))))
+		{
+			DWORD NativeProcessId;
+			if (SUCCEEDED(ActivationManager->ActivateApplication(Aumid, Params, AO_NONE, &NativeProcessId)))
+			{
+				if (OutProcessId != nullptr)
+				{
+					*OutProcessId = NativeProcessId;
+				}
+				ActivationSuccess = true;
+			}
+		}
+
+		CoTaskMemFree(Aumid);
+		return ActivationSuccess;
+#else
+		return false;
+#endif
+	}
+}
+
+FHoloLensTargetDevice::FHoloLensTargetDevice(const ITargetPlatform& InTargetPlatform, const FHoloLensDeviceInfo& InInfo)
+	: HeartBeatRequest(nullptr)
+	, TargetPlatform(InTargetPlatform)
+	, Info(InInfo)
+	, IsDeviceConnected(true)
+{
+	if (!Info.IsLocal)
+	{
+#if WITH_ENGINE
+		GWorld->GetTimerManager().SetTimer(TimerHandle, std::bind(&FHoloLensTargetDevice::StartHeartBeat, this), 10.f, true, 10.f);
+#endif
+	}
+}
+
+FHoloLensTargetDevice::~FHoloLensTargetDevice()
+{
+	FScopeLock lock(&CriticalSection);
+	if (!Info.IsLocal)
+	{
+#if WITH_ENGINE
+		if ((GWorld != nullptr) && TimerHandle.IsValid())
+		{
+			GWorld->GetTimerManager().ClearTimer(TimerHandle);
+		}
+#endif
+		if (HeartBeatRequest)
+		{
+			HeartBeatRequest->OnProcessRequestComplete().Unbind();
+			HeartBeatRequest->CancelRequest();
+			HeartBeatRequest = nullptr;
+		}
+	}
+}
+
+
+bool FHoloLensTargetDevice::Connect() 
+{
+	return true;
+}
+
+void FHoloLensTargetDevice::Disconnect() 
+{ }
+
+ETargetDeviceTypes FHoloLensTargetDevice::GetDeviceType() const 
+{
+	return ETargetDeviceTypes::HMD;
+}
+
+FTargetDeviceId FHoloLensTargetDevice::GetId() const 
+{
+	if (Info.IsLocal)
+	{
+		return FTargetDeviceId(TargetPlatform.PlatformName(), Info.HostName);
+	}
+	else
+	{
+		// This is what gets handed off to UAT, so we need to supply the
+		// actual Device Portal url instead of just the host name
+		return FTargetDeviceId(TargetPlatform.PlatformName(), Info.WdpUrl);
+	}
+}
+
+FString FHoloLensTargetDevice::GetName() const 
+{
+	return Info.HostName + TEXT(" (HoloLens)");
+}
+
+FString FHoloLensTargetDevice::GetOperatingSystemName() 
+{
+	return FString::Printf(TEXT("HoloLens (%s)"), *Info.DeviceTypeName.ToString());
+}
+
+bool FHoloLensTargetDevice::GetProcessSnapshotAsync(TFunction<void(const TArray<FTargetDeviceProcessInfo>&)> CompleteHandler)
+{
+	if (Info.IsLocal)
+	{
+		return false;
+	}
+
+	auto HttpRequest = GenerateRequest();
+	HttpRequest->SetVerb(TEXT("GET"));
+	HttpRequest->SetURL(Info.WdpUrl + TEXT("api/resourcemanager/processes"));
+	HttpRequest->OnProcessRequestComplete().BindLambda(
+		[&, CompleteHandler](FHttpRequestPtr HttpRequest, FHttpResponsePtr HttpResponse, bool bSucceeded)
+	{
+		TArray<FTargetDeviceProcessInfo> OutProcessInfos;
+		do
+		{
+			if (!bSucceeded || !HttpResponse.IsValid())
+			{
+				break;
+			}
+			FString Json = HttpResponse->GetContentAsString();
+			auto RootObject = TSharedPtr<FJsonObject>();
+			auto Reader = TJsonReaderFactory<TCHAR>::Create(Json);
+
+			if (!FJsonSerializer::Deserialize(Reader, RootObject))
+			{
+				break;
+			}
+
+			const TArray< TSharedPtr<FJsonValue> > * OutArray;
+			if (!RootObject->TryGetArrayField(TEXT("Processes"), OutArray))
+			{
+				break;
+			}
+
+			for (TSharedPtr<FJsonValue> process : *OutArray)
+			{
+				FTargetDeviceProcessInfo ProcessInfo;
+				ProcessInfo.ParentId = 0;
+				FString Publisher; //present for HoloLens processes only
+				FString ImageName, AppName;
+				TSharedPtr<FJsonObject> obj = process->AsObject();
+				if (!GetJsonField(Publisher, obj, TEXT("Publisher")))
+				{
+					continue;
+				}
+				if (!GetJsonField(ImageName, obj, TEXT("ImageName")))
+				{
+					continue;
+				}
+				if (!GetJsonField(AppName, obj, TEXT("AppName")))
+				{
+					continue;
+				}
+				if (!GetJsonField(ProcessInfo.UserName, obj, TEXT("UserName")))
+				{
+					continue;
+				}
+				if (!GetJsonField(ProcessInfo.Id, obj, TEXT("ProcessId")))
+				{
+					continue;
+				}
+				ProcessInfo.Name = FString::Format(TEXT("{0} ({1})"), { ImageName , AppName } );
+				OutProcessInfos.Add(ProcessInfo);
+			}
+		} while (false);
+
+		CompleteHandler(OutProcessInfos);
+	});
+
+	HttpRequest->ProcessRequest();
+	return true;
+}
+
+const class ITargetPlatform& FHoloLensTargetDevice::GetTargetPlatform() const
+{
+	return TargetPlatform;
+}
+
+bool FHoloLensTargetDevice::GetUserCredentials(FString& OutUserName, FString& OutUserPassword) 
+{
+	if (!Info.RequiresCredentials)
+	{
+		return false;
+	}
+
+	OutUserName = Info.Username;
+	OutUserPassword = Info.Password;
+	return true;
+}
+
+bool FHoloLensTargetDevice::IsConnected()
+{
+	if (Info.IsLocal)
+	{
+		return true;
+	}
+
+	return IsDeviceConnected;
+}
+
+bool FHoloLensTargetDevice::IsDefault() const 
+{
+	return true;
+}
+
+TSharedPtr<IHttpRequest, ESPMode::ThreadSafe> FHoloLensTargetDevice::GenerateRequest() const
+{
+	if (Info.IsLocal)
+	{
+		return nullptr;
+	}
+
+	SslCertDisabler disabler;
+	
+	auto HttpRequest = TSharedPtr<IHttpRequest, ESPMode::ThreadSafe>(FPlatformHttp::ConstructRequest());
+
+	if (Info.RequiresCredentials)
+	{
+		FString AuthentificationString = FString::Format(TEXT("auto-{0}:{1}"), { Info.Username, Info.Password });
+		FString AuthHeaderString = FString::Format(TEXT("Basic {0}"), { FBase64::Encode(AuthentificationString) });
+		HttpRequest->AppendToHeader(TEXT("Authorization"), AuthHeaderString);
+	}
+
+	return HttpRequest;
+}
+
+
+bool FHoloLensTargetDevice::PowerOff(bool Force) 
+{
+	if (Info.IsLocal)
+	{
+		return false;
+	}
+
+	auto HttpRequest = GenerateRequest();
+	HttpRequest->SetVerb(TEXT("POST"));
+	HttpRequest->SetURL(Info.WdpUrl + TEXT("api/control/shutdown"));
+	HttpRequest->ProcessRequest();
+	return true;
+}
+
+bool FHoloLensTargetDevice::PowerOn() 
+{
+	return false;
+}
+
+bool FHoloLensTargetDevice::Reboot(bool bReconnect) 
+{
+	if (Info.IsLocal)
+	{
+		return false;
+	}
+
+	auto HttpRequest = GenerateRequest();
+	HttpRequest->SetVerb(TEXT("POST"));
+	HttpRequest->SetURL(Info.WdpUrl + TEXT("api/control/restart"));
+	HttpRequest->ProcessRequest();
+
+	return true;
+}
+
+void FHoloLensTargetDevice::SetUserCredentials(const FString& UserName, const FString& UserPassword) 
+{ 
+	Info.Username = UserName;
+	Info.Password = UserPassword;
+}
+
+bool FHoloLensTargetDevice::SupportsFeature(ETargetDeviceFeatures Feature) const 
+{
+	if (Info.IsLocal)
+	{
+		return false;
+	}
+
+	switch (Feature)
+	{
+	case ETargetDeviceFeatures::PowerOff:
+	case ETargetDeviceFeatures::ProcessSnapshot:
+	case ETargetDeviceFeatures::Reboot:
+		return true;
+	default:
+		return false;
+	}
+}
+
+bool FHoloLensTargetDevice::TerminateProcess(const int64 ProcessId) 
+{
+	if (Info.IsLocal)
+	{
+		return false;
+	}
+
+	auto HttpRequest = GenerateRequest();
+	HttpRequest->SetVerb(TEXT("DELETE"));
+	HttpRequest->SetURL(Info.WdpUrl + TEXT("api/taskmanager/process?pid=") + FString::FromInt(ProcessId));
+	HttpRequest->ProcessRequest();
+
+	return true;
+}
+
+bool FHoloLensTargetDevice::TerminateLaunchedProcess(const FString & ProcessIdentifier)
+{
+	if (Info.IsLocal)
+	{
+		return false;
+	}
+
+	auto HttpRequest = GenerateRequest();
+	HttpRequest->SetVerb(TEXT("DELETE"));
+	HttpRequest->SetURL(Info.WdpUrl + TEXT("api/taskmanager/app?package=") + FBase64::Encode(ProcessIdentifier));
+	HttpRequest->ProcessRequest();
+
+	return true;
+}
+
+void FHoloLensTargetDevice::StartHeartBeat()
+{
+	FScopeLock lock(&CriticalSection);
+	if (HeartBeatRequest)
+	{
+		return;
+	}
+#if WITH_EDITOR
+	if (!TimerHandle.IsValid()) // the app is quiting
+	{
+		return;
+	}
+#endif
+	HeartBeatRequest = GenerateRequest();
+	HeartBeatRequest->SetVerb(TEXT("GET"));
+	HeartBeatRequest->SetURL(Info.WdpUrl + TEXT("api/os/info"));
+
+	HeartBeatRequest->OnProcessRequestComplete().BindLambda(
+		[=](FHttpRequestPtr HttpRequest, FHttpResponsePtr HttpResponse, bool bSucceeded)
+	{
+		FString HostName;
+		bool Success = false;
+		do
+		{
+			if (!bSucceeded || !HttpResponse.IsValid())
+			{
+				break;
+			}
+			FString Json = HttpResponse->GetContentAsString();
+			auto RootObject = TSharedPtr<FJsonObject>();
+			auto Reader = TJsonReaderFactory<TCHAR>::Create(Json);
+
+			if (!FJsonSerializer::Deserialize(Reader, RootObject))
+			{
+				break;
+			}
+
+			if (GetJsonField(HostName, RootObject, TEXT("ComputerName")))
+			{
+				Success = (HostName == Info.HostName);
+			}
+		} while (false);
+
+		{
+			FScopeLock lock(&CriticalSection);
+			IsDeviceConnected = Success;
+			HeartBeatRequest = nullptr;
+		}
+	});
+
+	HeartBeatRequest->ProcessRequest();
+}
+
+
+#include "Windows/HideWindowsPlatformTypes.h"
diff --git a/Engine/Source/Developer/HoloLens/HoloLensTargetPlatform/Private/HoloLensTargetDevice.h b/Engine/Source/Developer/HoloLens/HoloLensTargetPlatform/Private/HoloLensTargetDevice.h
new file mode 100644
--- /dev/null
+++ b/Engine/Source/Developer/HoloLens/HoloLensTargetPlatform/Private/HoloLensTargetDevice.h
@@ -0,0 +1,100 @@
+// Copyright Epic Games, Inc. All Rights Reserved.
+
+/*=============================================================================
+	HoloLensTargetDevice.h: Declares the HoloLensTargetDevice class.
+=============================================================================*/
+
+#pragma once
+
+#include "CoreMinimal.h"
+#include "HAL/PlatformProcess.h"
+#include "Interfaces/ITargetDevice.h"
+#include "Interfaces/ITargetPlatform.h"
+#include "IHoloLensDeviceDetector.h"
+
+#if WITH_ENGINE
+#include "Engine/EngineTypes.h"
+#endif
+
+#include "Windows/AllowWindowsPlatformTypes.h"
+
+
+/**
+ * Implements a HoloLens target device.
+ */
+class FHoloLensTargetDevice
+	: public ITargetDevice
+{
+public:
+
+	/**
+	* Creates and initializes a new device for the specified target platform.
+	*
+	* @param InTargetPlatform - The target platform.
+	*/
+	FHoloLensTargetDevice(const ITargetPlatform& InTargetPlatform, const FHoloLensDeviceInfo& InInfo);
+
+	~FHoloLensTargetDevice();
+
+	virtual bool Connect() override;
+
+	virtual void Disconnect() override;
+
+	virtual ETargetDeviceTypes GetDeviceType() const override;
+
+	virtual FTargetDeviceId GetId() const override;
+
+	virtual FString GetName() const override;
+
+	virtual FString GetOperatingSystemName() override;
+
+	virtual bool GetProcessSnapshotAsync(TFunction<void(const TArray<FTargetDeviceProcessInfo>&)> CompleteHandler) override;
+
+	virtual int32 GetProcessSnapshot(TArray<FTargetDeviceProcessInfo>& OutProcessInfos) override { return 0; }
+
+	virtual const class ITargetPlatform& GetTargetPlatform() const override;
+
+	virtual bool GetUserCredentials(FString& OutUserName, FString& OutUserPassword) override;
+
+	virtual bool IsConnected();
+
+	virtual bool IsDefault() const override;
+
+	virtual bool PowerOff(bool Force) override;
+
+	virtual bool PowerOn() override;
+
+	virtual bool Reboot(bool bReconnect = false) override;
+
+	virtual void SetUserCredentials(const FString& UserName, const FString& UserPassword) override;
+
+	virtual bool SupportsFeature(ETargetDeviceFeatures Feature) const override;
+
+	virtual bool TerminateProcess(const int64 ProcessId) override;
+
+	virtual bool TerminateLaunchedProcess(const FString & ProcessIdentifier) override;
+
+private:
+
+	TSharedPtr<class IHttpRequest, ESPMode::ThreadSafe> GenerateRequest() const;
+
+	void StartHeartBeat();
+
+#if WITH_ENGINE
+	FTimerHandle TimerHandle;
+#endif
+
+	TSharedPtr<class IHttpRequest, ESPMode::ThreadSafe> HeartBeatRequest;
+	FCriticalSection CriticalSection;
+
+	// Holds a reference to the device's target platform.
+	const ITargetPlatform& TargetPlatform;
+
+	FHoloLensDeviceInfo Info;
+
+	volatile bool IsDeviceConnected;
+};
+
+typedef TSharedPtr<FHoloLensTargetDevice, ESPMode::ThreadSafe> FHoloLensDevicePtr;
+
+#include "Windows/HideWindowsPlatformTypes.h"
diff --git a/Engine/Source/Developer/HoloLens/HoloLensTargetPlatform/Private/HoloLensTargetPlatform.cpp b/Engine/Source/Developer/HoloLens/HoloLensTargetPlatform/Private/HoloLensTargetPlatform.cpp
new file mode 100644
--- /dev/null
+++ b/Engine/Source/Developer/HoloLens/HoloLensTargetPlatform/Private/HoloLensTargetPlatform.cpp
@@ -0,0 +1,209 @@
+// Copyright Epic Games, Inc. All Rights Reserved.
+
+/*=============================================================================
+	HoloLensTargetPlatform.cpp: Implements the FHoloLensTargetPlatform class.
+=============================================================================*/
+
+#include "HoloLensTargetPlatform.h"
+#include "HoloLensTargetDevice.h"
+#include "Misc/ConfigCacheIni.h"
+#include "Misc/ScopeLock.h"
+#include "HttpModule.h"
+#include "PlatformHttp.h"
+#include "Interfaces/IHttpResponse.h"
+#include "HoloLensPlatformEditor.h"
+#include "GeneralProjectSettings.h"
+
+DEFINE_LOG_CATEGORY(LogHoloLensTargetPlatform);
+
+FHoloLensTargetPlatform::FHoloLensTargetPlatform(bool bIsClientOnly)
+	: TNonDesktopTargetPlatformBase(bIsClientOnly)
+{
+#if WITH_ENGINE
+	TextureLODSettings = nullptr; // These are registered by the device profile system.
+	StaticMeshLODSettings.Initialize(this);
+#endif
+
+	DeviceDetector = IHoloLensDeviceDetector::Create();
+
+	DeviceDetectedRegistration = DeviceDetector->OnDeviceDetected().AddRaw(this, &FHoloLensTargetPlatform::OnDeviceDetected);
+
+	{
+		DeviceDetector->StartDeviceDetection();
+	}
+}
+
+FHoloLensTargetPlatform::~FHoloLensTargetPlatform()
+{
+	DeviceDetector->OnDeviceDetected().Remove(DeviceDetectedRegistration);
+}
+
+void FHoloLensTargetPlatform::GetAllDevices(TArray<ITargetDevicePtr>& OutDevices) const
+{
+	DeviceDetector->StartDeviceDetection();
+
+	OutDevices.Reset();
+	FScopeLock Lock(&DevicesLock);
+	OutDevices = Devices;
+}
+
+ITargetDevicePtr FHoloLensTargetPlatform::GetDevice(const FTargetDeviceId& DeviceId)
+{
+	if (PlatformName() == DeviceId.GetPlatformName())
+	{
+		DeviceDetector->StartDeviceDetection();
+
+		FScopeLock Lock(&DevicesLock);
+		for (ITargetDevicePtr Device : Devices)
+		{
+			if (DeviceId == Device->GetId())
+			{
+				return Device;
+			}
+		}
+	}
+
+
+	return nullptr;
+}
+
+ITargetDevicePtr FHoloLensTargetPlatform::GetDefaultDevice() const
+{
+	DeviceDetector->StartDeviceDetection();
+
+	FScopeLock Lock(&DevicesLock);
+	for (ITargetDevicePtr RemoteDevice : Devices)
+	{
+		if (RemoteDevice->IsDefault())
+		{
+			return RemoteDevice;
+		}
+	}
+
+	return nullptr;
+}
+
+bool FHoloLensTargetPlatform::SupportsFeature(ETargetPlatformFeatures Feature) const
+{
+	switch (Feature)
+	{
+	case ETargetPlatformFeatures::Packaging:
+	case ETargetPlatformFeatures::UserCredentials:
+	case ETargetPlatformFeatures::DeviceOutputLog:
+		return true;
+	default:
+		return TTargetPlatformBase<FHoloLensPlatformProperties>::SupportsFeature(Feature);
+	}
+}
+
+// Hololens previously used PCD3D_ES3_1 as a shader format/platform until we
+// had to split to properly support DataDrivenPlatformInfo
+static FName NAME_D3D_ES3_1_HOLOLENS(TEXT("D3D_ES3_1_HOLOLENS"));
+static FName NAME_PCD3D_SM5(TEXT("PCD3D_SM5"));
+
+void FHoloLensTargetPlatform::GetAllPossibleShaderFormats(TArray<FName>& OutFormats) const
+{
+	OutFormats.AddUnique(NAME_D3D_ES3_1_HOLOLENS);
+	OutFormats.AddUnique(NAME_PCD3D_SM5);
+}
+
+void FHoloLensTargetPlatform::GetAllTargetedShaderFormats(TArray<FName>& OutFormats) const
+{
+	OutFormats.AddUnique(NAME_D3D_ES3_1_HOLOLENS);
+	OutFormats.AddUnique(NAME_PCD3D_SM5);
+}
+
+#if WITH_ENGINE
+
+void FHoloLensTargetPlatform::GetReflectionCaptureFormats(TArray<FName>& OutFormats) const
+{
+	OutFormats.Add(FName(TEXT("FullHDR")));
+	OutFormats.Add(FName(TEXT("EncodedHDR")));
+}
+
+void FHoloLensTargetPlatform::GetTextureFormats(const UTexture* InTexture, TArray< TArray<FName> >& OutFormats) const
+{
+	GetDefaultTextureFormatNamePerLayer(OutFormats.AddDefaulted_GetRef(), this, InTexture, false);
+}
+
+void FHoloLensTargetPlatform::GetAllTextureFormats(TArray<FName>& OutFormats) const
+{
+	GetAllDefaultTextureFormats(this, OutFormats, false);
+}
+
+#endif
+
+void FHoloLensTargetPlatform::OnDeviceDetected(const FHoloLensDeviceInfo& Info)
+{
+	FHoloLensDevicePtr NewDevice = MakeShared<FHoloLensTargetDevice, ESPMode::ThreadSafe>(*this, Info);
+	{
+		FScopeLock Lock(&DevicesLock);
+		Devices.Add(NewDevice);
+	}
+	OnDeviceDiscovered().Broadcast(NewDevice.ToSharedRef());
+}
+
+bool FHoloLensTargetPlatform::SupportsBuildTarget(EBuildTargetType BuildTarget) const
+{
+	return BuildTarget == EBuildTargetType::Game;
+}
+
+bool FHoloLensTargetPlatform::IsSdkInstalled(bool bProjectHasCode, FString& OutDocumentationPath) const
+{
+	OutDocumentationPath = TEXT("Platforms/HoloLens/GettingStarted");
+
+	const TArray<FHoloLensSDKVersion>& SDKVersions = FHoloLensSDKVersion::GetSDKVersions();
+	return SDKVersions.Num() > 0;
+}
+
+int32 FHoloLensTargetPlatform::CheckRequirements(bool bProjectHasCode, EBuildConfiguration Configuration, bool bRequiresAssetNativization, FString& OutTutorialPath, FString& OutDocumentationPath, FText& CustomizedLogMessage) const
+{
+	OutDocumentationPath = TEXT("Platforms/HoloLens/GettingStarted");
+	FString LocalErrors;
+
+	int32 BuildStatus = ETargetPlatformReadyStatus::Ready;
+	if (!IsSdkInstalled(bProjectHasCode, OutTutorialPath))
+	{
+		BuildStatus |= ETargetPlatformReadyStatus::SDKNotFound;
+	}
+	FString PublisherIdentityName = GetDefault<UGeneralProjectSettings>()->CompanyDistinguishedName;
+	if (PublisherIdentityName.IsEmpty())
+	{
+		LocalErrors += TEXT("Missing Company Distinguished Name (See Project Settings).");
+		BuildStatus |= ETargetPlatformReadyStatus::SigningKeyNotFound;
+	}
+	else
+	{
+		if (PublisherIdentityName.Contains(TEXT("CN=")) && PublisherIdentityName.Len() == 3)
+		{
+			LocalErrors += TEXT(" Malformed Company Distinguished Name (See Project Settings).");
+			BuildStatus |= ETargetPlatformReadyStatus::SigningKeyNotFound;
+		}
+	}
+	FString ProjectName = GetDefault<UGeneralProjectSettings>()->ProjectName;
+	if (ProjectName.IsEmpty())
+	{
+		LocalErrors += TEXT(" Missing Project Name (See Project Settings).");
+		BuildStatus |= ETargetPlatformReadyStatus::SigningKeyNotFound;
+	}
+
+	// Set the path if missing any of the bits needed for signing
+	if (BuildStatus & ETargetPlatformReadyStatus::SigningKeyNotFound)
+	{
+		OutDocumentationPath = TEXT("Platforms/HoloLens/Signing");
+	}
+
+	if (BuildStatus != ETargetPlatformReadyStatus::Ready)
+	{
+		UE_LOG(LogHoloLensTargetPlatform, Warning, TEXT("FHoloLensTargetPlatform::CheckRequirements found these problems: %s"), *LocalErrors);
+	}
+
+	return BuildStatus;
+}
+
+bool FHoloLensTargetPlatform::AddDevice(const FString& DeviceId, const FString& DeviceUserFriendlyName, const FString& Username, const FString& Password, bool bDefault)
+{
+	DeviceDetector->TryAddDevice(DeviceId, DeviceUserFriendlyName, Username, Password);
+
+	return true;
+}
diff --git a/Engine/Source/Developer/HoloLens/HoloLensTargetPlatform/Private/HoloLensTargetPlatform.h b/Engine/Source/Developer/HoloLens/HoloLensTargetPlatform/Private/HoloLensTargetPlatform.h
new file mode 100644
--- /dev/null
+++ b/Engine/Source/Developer/HoloLens/HoloLensTargetPlatform/Private/HoloLensTargetPlatform.h
@@ -0,0 +1,138 @@
+// Copyright Epic Games, Inc. All Rights Reserved.
+
+/*=============================================================================
+	HoloLensTargetPlatform.h: Declares the FXboxOneTargetPlatform class.
+=============================================================================*/
+
+#pragma once
+
+#include "Common/TargetPlatformBase.h"
+#include "Runtime/Core/Public/HoloLens/HoloLensPlatformProperties.h"
+#include "Misc/ConfigCacheIni.h"
+#include "HoloLensTargetDevice.h"
+#include "Misc/ScopeLock.h"
+#include "IHoloLensDeviceDetector.h"
+
+#if PLATFORM_WINDOWS || PLATFORM_HOLOLENS
+#include "Windows/AllowWindowsPlatformTypes.h"
+#endif
+
+#if WITH_ENGINE
+#include "StaticMeshResources.h"
+#endif // WITH_ENGINE
+
+#define LOCTEXT_NAMESPACE "HoloLensTargetPlatform"
+
+/**
+ * FHoloLensTargetPlatform, abstraction for cooking HoloLens platforms
+ */
+class HOLOLENSTARGETPLATFORM_API FHoloLensTargetPlatform
+	: public TNonDesktopTargetPlatformBase<FHoloLensPlatformProperties>
+{
+public:
+
+	/**
+	 * Default constructor.
+	 */
+	FHoloLensTargetPlatform(bool bIsClient);
+
+	/**
+	 * Destructor.
+	 */
+	virtual ~FHoloLensTargetPlatform();
+
+public:
+
+	//~ Begin ITargetPlatform Interface
+
+	virtual bool AddDevice(const FString& DeviceId, const FString& DeviceUserFriendlyName, const FString& Username, const FString& Password, bool bDefault) override;
+
+	virtual void EnableDeviceCheck(bool OnOff) override {}
+
+	virtual void GetAllDevices(TArray<ITargetDevicePtr>& OutDevices) const override;
+
+	virtual ITargetDevicePtr GetDefaultDevice() const override;
+
+	virtual ITargetDevicePtr GetDevice(const FTargetDeviceId& DeviceId) override;
+
+	//virtual ECompressionFlags GetBaseCompressionMethod() const override { return ECompressionFlags::COMPRESS_ZLIB; }
+
+	virtual bool GenerateStreamingInstallManifest(const TMultiMap<FString, int32>& ChunkMap, const TSet<int32>& ChunkIDsInUse) const override { return true; }
+
+	virtual bool SupportsFeature(ETargetPlatformFeatures Feature) const override;
+
+	virtual bool SupportsBuildTarget(EBuildTargetType BuildTarget) const override;
+
+	virtual void GetAllPossibleShaderFormats(TArray<FName>& OutFormats) const override;
+
+	virtual void GetAllTargetedShaderFormats(TArray<FName>& OutFormats) const override;	
+
+#if WITH_ENGINE
+	virtual void GetReflectionCaptureFormats(TArray<FName>& OutFormats) const override;
+
+	virtual const class FStaticMeshLODSettings& GetStaticMeshLODSettings() const override { return StaticMeshLODSettings; }
+
+	virtual void GetTextureFormats(const UTexture* InTexture, TArray< TArray<FName> >& OutFormats) const override;
+
+	virtual const UTextureLODSettings& GetTextureLODSettings() const override { return *TextureLODSettings; }
+
+	virtual void GetAllTextureFormats(TArray<FName>& OutFormats) const override;
+
+	virtual void RegisterTextureLODSettings(const UTextureLODSettings* InTextureLODSettings) override
+	{
+		TextureLODSettings = InTextureLODSettings;
+	}
+
+	virtual FName GetWaveFormat( const class USoundWave* Wave ) const override
+	{
+		static FName NAME_OGG(TEXT("OGG"));
+
+		return NAME_OGG;
+	}
+
+	virtual void GetAllWaveFormats(TArray<FName>& OutFormats) const override
+	{
+		static FName NAME_OGG(TEXT("OGG"));
+		static FName NAME_OPUS(TEXT("OPUS"));
+		OutFormats.Add(NAME_OGG);
+		OutFormats.Add(NAME_OPUS);
+	}
+
+#endif //WITH_ENGINE
+
+	virtual EPlatformAuthentication RequiresUserCredentials() const override
+	{
+		return EPlatformAuthentication::Possible;
+	}
+
+	virtual bool IsSdkInstalled(bool bProjectHasCode, FString& OutDocumentationPath) const override;
+
+	virtual int32 CheckRequirements(bool bProjectHasCode, EBuildConfiguration Configuration, bool bRequiresAssetNativization, FString& OutTutorialPath, FString& OutDocumentationPath, FText& CustomizedLogMessage) const override;
+	//~ End ITargetPlatform Interface
+
+private:
+	IHoloLensDeviceDetectorPtr DeviceDetector;
+
+	void OnDeviceDetected(const FHoloLensDeviceInfo& Info);
+
+	mutable FCriticalSection DevicesLock;
+	TArray<ITargetDevicePtr> Devices;
+
+	FDelegateHandle DeviceDetectedRegistration;
+
+#if WITH_ENGINE
+	// Holds a cache of the target LOD settings.
+	const UTextureLODSettings* TextureLODSettings;
+
+	// Holds the static mesh LOD settings.
+	FStaticMeshLODSettings StaticMeshLODSettings;
+#endif //WITH_ENGINE
+
+};
+
+
+#undef LOCTEXT_NAMESPACE
+
+#include "Windows/HideWindowsPlatformTypes.h"
+
+DECLARE_LOG_CATEGORY_EXTERN(LogHoloLensTargetPlatform, Log, All);
diff --git a/Engine/Source/Developer/HoloLens/HoloLensTargetPlatform/Private/HoloLensTargetPlatformModule.cpp b/Engine/Source/Developer/HoloLens/HoloLensTargetPlatform/Private/HoloLensTargetPlatformModule.cpp
new file mode 100644
--- /dev/null
+++ b/Engine/Source/Developer/HoloLens/HoloLensTargetPlatform/Private/HoloLensTargetPlatformModule.cpp
@@ -0,0 +1,40 @@
+// Copyright Epic Games, Inc. All Rights Reserved.
+
+#include "HoloLensTargetDevice.h"
+#include "Interfaces/ITargetPlatformModule.h"
+#include "HoloLensTargetPlatform.h"
+#if PLATFORM_WINDOWS || PLATFORM_HOLOLENS
+#include "Windows/AllowWindowsPlatformTypes.h"
+#endif
+#include "ISettingsModule.h"
+#include "Modules/ModuleManager.h"
+#include "UObject/Package.h"
+
+#define LOCTEXT_NAMESPACE "FHoloLensTargetPlatformModule"
+
+
+/**
+ * Module for the HoloLens target platform.
+ */
+class FHoloLensTargetPlatformModule
+	: public ITargetPlatformModule
+{
+public:
+	virtual void GetTargetPlatforms(TArray<ITargetPlatform*>& TargetPlatforms) override
+	{
+		if (FHoloLensTargetPlatform::IsUsable())
+		{
+			//@todo HoloLens: Check for SDK?
+			TargetPlatforms.Add(new FHoloLensTargetPlatform(false));
+			TargetPlatforms.Add(new FHoloLensTargetPlatform(true));
+		}
+	}
+};
+
+#undef LOCTEXT_NAMESPACE
+
+
+IMPLEMENT_MODULE(FHoloLensTargetPlatformModule, HoloLensTargetPlatform);
+
+
+#include "Windows/HideWindowsPlatformTypes.h"
diff --git a/Engine/Source/Developer/HoloLens/HoloLensTargetPlatform/Private/IHoloLensDeviceDetector.h b/Engine/Source/Developer/HoloLens/HoloLensTargetPlatform/Private/IHoloLensDeviceDetector.h
new file mode 100644
--- /dev/null
+++ b/Engine/Source/Developer/HoloLens/HoloLensTargetPlatform/Private/IHoloLensDeviceDetector.h
@@ -0,0 +1,67 @@
+// Copyright Epic Games, Inc. All Rights Reserved.
+
+#pragma once
+
+#include "Modules/ModuleInterface.h"
+#include "Modules/ModuleManager.h"
+#include "HAL/PlatformProcess.h"
+
+class FJsonObject;
+
+namespace HoloLensDeviceTypes
+{
+	extern const FName HoloLens;
+	extern const FName HoloLensEmulation;
+}
+
+enum class EHoloLensArchitecture : uint8
+{
+	Unknown = 0, X86 = 1, X64 = 2, ARM32 = 3, ARM64 = 4
+};
+
+
+struct FHoloLensDeviceInfo
+{
+	FString HostName;
+	FString WdpUrl;
+	FString WindowsDeviceId;
+	FName DeviceTypeName;
+	EHoloLensArchitecture Architecture;
+	FString Username;
+	FString Password;
+	int32 RequiresCredentials : 1;
+	int32 IsLocal : 1;
+};
+
+typedef TSharedPtr<class IHoloLensDeviceDetector, ESPMode::ThreadSafe> IHoloLensDeviceDetectorPtr;
+
+class IHoloLensDeviceDetector
+{
+public:
+	// Virtual destructor.
+	virtual ~IHoloLensDeviceDetector() {};
+
+	static IHoloLensDeviceDetectorPtr Create();
+
+	virtual void StartDeviceDetection() = 0;
+	virtual void StopDeviceDetection() = 0;
+
+	DECLARE_EVENT_OneParam(IHoloLensDeviceDetectorModule, FOnDeviceDetected, const FHoloLensDeviceInfo&);
+	virtual FOnDeviceDetected& OnDeviceDetected() = 0;
+
+	virtual const TArray<FHoloLensDeviceInfo> GetKnownDevices() = 0;
+
+	virtual void TryAddDevice(const FString& DeviceId, const FString& DeviceUserFriendlyName, const FString& Username, const FString& Password) = 0;
+};
+
+bool GetJsonField(FString& OutVal, const TSharedPtr<FJsonObject>& JsonObject, const TCHAR* FieldName);
+bool GetJsonField(int64& OutVal, const TSharedPtr<FJsonObject>& JsonObject, const TCHAR* FieldName);
+
+class SslCertDisabler
+{
+	bool prevValue;
+public:
+	SslCertDisabler();
+
+	~SslCertDisabler();
+};
diff --git a/Engine/Source/Programs/UnrealBuildTool/Platform/HoloLens/UEBuildHoloLens.cs b/Engine/Source/Programs/UnrealBuildTool/Platform/HoloLens/UEBuildHoloLens.cs
--- a/Engine/Source/Programs/UnrealBuildTool/Platform/HoloLens/UEBuildHoloLens.cs
+++ b/Engine/Source/Programs/UnrealBuildTool/Platform/HoloLens/UEBuildHoloLens.cs
@@ -386,7 +386,41 @@ namespace UnrealBuildTool
 		/// <param name="Target">The target being build</param>
 		public override void ModifyModuleRulesForOtherPlatform(string ModuleName, ModuleRules Rules, ReadOnlyTargetRules Target)
 		{
-			// This code has been removed because it causes a full rebuild after generating project files (since response files are overwritten with different defines).
+			// don't do any target platform stuff if SDK is not available
+			if (!UEBuildPlatform.IsPlatformAvailableForTarget(Platform, Target))
+			{
+				return;
+			}
+			
+			if (Target.Platform == UnrealTargetPlatform.Win64)
+			{
+				if (!Target.bBuildRequiresCookedData)
+				{
+					if (ModuleName == "Engine")
+					{
+						if (Target.bBuildDeveloperTools)
+						{
+							Rules.DynamicallyLoadedModuleNames.Add("HoloLensTargetPlatform");
+						}
+					}
+				}
+
+				// allow standalone tools to use targetplatform modules, without needing Engine
+				if (ModuleName == "TargetPlatform")
+				{
+					if (Target.bForceBuildTargetPlatforms)
+					{
+						Rules.DynamicallyLoadedModuleNames.Add("HoloLensTargetPlatform");
+					}
+				}
+
+				if (ModuleName == "UnrealEd")
+				{
+					Rules.DynamicallyLoadedModuleNames.Add("HoloLensPlatformEditor");
+				}
+			}
+			
+// This code has been removed because it causes a full rebuild after generating project files (since response files are overwritten with different defines).
 #if false
 			if (Target.Platform == UnrealTargetPlatform.Win64)
 			{
diff --git a/Engine/Source/Runtime/Core/Public/HoloLens/HoloLensPlatformProperties.h b/Engine/Source/Runtime/Core/Public/HoloLens/HoloLensPlatformProperties.h
--- a/Engine/Source/Runtime/Core/Public/HoloLens/HoloLensPlatformProperties.h
+++ b/Engine/Source/Runtime/Core/Public/HoloLens/HoloLensPlatformProperties.h
@@ -63,6 +63,11 @@ struct FHoloLensPlatformProperties : public FGenericPlatformProperties
 		// HoloLens 2 is mobile renderer, thus only supports low quality light maps
 		return false; 
 	}
+
+	static FORCEINLINE bool SupportsVirtualTextureStreaming()
+	{
+		return true;
+	}
 };
 
 #ifdef PROPERTY_HEADER_SHOULD_DEFINE_TYPE
diff --git a/Templates/TP_HololensBP/TP_HololensBP.uproject b/Templates/TP_HololensBP/TP_HololensBP.uproject
--- a/Templates/TP_HololensBP/TP_HololensBP.uproject
+++ b/Templates/TP_HololensBP/TP_HololensBP.uproject
@@ -57,13 +57,6 @@
 				"Win64"
 			]
 		},
-		{
-			"Name": "HoloLensAR",
-			"Enabled": true,
-			"SupportedTargetPlatforms": [
-				"Win64"
-			]
-		},
 		{
 			"Name": "OpenXRHandTracking",
 			"Enabled": true,
