UE-142711 UnrealInsight hololens build break
Reversing the dependency between HololensPlatformEditor and HololensTargetPlatform.  Editor should depend on Target, not the reverse.

diff --git a/Engine/Source/Developer/HoloLens/HoloLensPlatformEditor/HoloLensPlatformEditor.Build.cs b/Engine/Source/Developer/HoloLens/HoloLensPlatformEditor/HoloLensPlatformEditor.Build.cs
--- a/Engine/Source/Developer/HoloLens/HoloLensPlatformEditor/HoloLensPlatformEditor.Build.cs
+++ b/Engine/Source/Developer/HoloLens/HoloLensPlatformEditor/HoloLensPlatformEditor.Build.cs
@@ -24,6 +24,7 @@ public class HoloLensPlatformEditor : ModuleRules
 				"UnrealEd",
 				"SourceControl",
 				"WindowsTargetPlatform", // For ECompilerVersion
+				"HoloLensTargetPlatform",
 				"EngineSettings",
 				"Projects",
 				"gltfToolkit",
diff --git a/Engine/Source/Developer/HoloLens/HoloLensPlatformEditor/Private/HoloLensTargetSettingsCustomization.cpp b/Engine/Source/Developer/HoloLens/HoloLensPlatformEditor/Private/HoloLensTargetSettingsCustomization.cpp
--- a/Engine/Source/Developer/HoloLens/HoloLensPlatformEditor/Private/HoloLensTargetSettingsCustomization.cpp
+++ b/Engine/Source/Developer/HoloLens/HoloLensPlatformEditor/Private/HoloLensTargetSettingsCustomization.cpp
@@ -1,7 +1,7 @@
 // Copyright Epic Games, Inc. All Rights Reserved.
 
 #include "HoloLensTargetSettingsCustomization.h"
-#include "HoloLensPlatformEditor.h"
+#include "HoloLensTargetTypes.h"
 #include "SExternalImageReference.h"
 #include "IExternalImagePickerModule.h"
 #include "DetailLayoutBuilder.h"
@@ -35,69 +35,6 @@
 
 #define LOCTEXT_NAMESPACE "HoloLensTargetSettingsCustomization"
 
-TArray<FHoloLensSDKVersion> FHoloLensSDKVersion::GetSDKVersions()
-{
-	bool bFound18362 = false;
-
-	static TArray<FHoloLensSDKVersion> SDKVersions;
-	if (SDKVersions.Num() == 0)
-	{
-		HKEY hRoots[] = { HKEY_LOCAL_MACHINE , HKEY_CURRENT_USER };
-
-		for (DWORD dwRootIdx = 0; dwRootIdx < _countof(hRoots); ++dwRootIdx)
-		{
-			HKEY hKey;
-			if (RegOpenKeyExW(hRoots[dwRootIdx], L"Software\\Microsoft\\Windows Kits\\Installed Roots", 0, KEY_ENUMERATE_SUB_KEYS | KEY_QUERY_VALUE, &hKey) == ERROR_SUCCESS)
-			{
-				DWORD dwIndex = 0;
-				WCHAR sSubKeyName[256];
-				DWORD dwSubKeyNameLen = _countof(sSubKeyName);
-				FString SDKPath;
-
-				if (ERROR_SUCCESS == RegQueryValueExW(hKey, L"KitsRoot10", 0, NULL, (LPBYTE)sSubKeyName, &dwSubKeyNameLen))
-				{
-					SDKPath = FString(dwSubKeyNameLen, sSubKeyName);
-				}
-				dwSubKeyNameLen = _countof(sSubKeyName);
-
-				while (RegEnumKeyExW(hKey, dwIndex, sSubKeyName, &dwSubKeyNameLen, NULL, NULL, NULL, NULL) == ERROR_SUCCESS)
-				{
-					int version[4];
-					if (swscanf_s(sSubKeyName, L"%d.%d.%d.%d", &version[0], &version[1], &version[2], &version[3]) == 4)
-					{
-						if (FPaths::FileExists(FPaths::Combine(SDKPath, TEXT("Include"), FString(dwSubKeyNameLen, sSubKeyName), TEXT("um"), TEXT("Windows.h"))))
-						{
-							// HoloLens min version is 17763 for remoting and 18317 for hand tracking
-							// @todo JoeG - update this version when the official os/sdk are released
-							if (version[0] >= 10 && version[2] >= 17763)
-							{
-								new(SDKVersions) FHoloLensSDKVersion(FString(dwSubKeyNameLen, sSubKeyName), version[0], version[1], version[2], version[3]);
-							}
-							if (version[2] == 18362)
-							{
-								bFound18362 = true;
-							}
-						}
-					}
-					dwSubKeyNameLen = _countof(sSubKeyName);
-					dwIndex++;
-				}
-				RegCloseKey(hKey);
-			}
-		}
-
-		// 18362 is required as the official minimum SDK released by Microsoft for all needed HL2 features to work
-		// If it's not found, we still need it as a hidden SDK version so that we can set MinimumPlatformVersion correctly to
-		// allow the package to run on lower installed firmware revisions
-		if (!bFound18362)
-		{
-			new(SDKVersions) FHoloLensSDKVersion(FString(L"10.0.18362.0"), 10, 0, 18362, 0);
-		}
-	}
-
-	return SDKVersions;
-}
-
 TSharedRef<IDetailCustomization> FHoloLensTargetSettingsCustomization::MakeInstance()
 {
 	return MakeShareable(new FHoloLensTargetSettingsCustomization);
diff --git a/Engine/Source/Developer/HoloLens/HoloLensPlatformEditor/Public/HoloLensPlatformEditor.h b/Engine/Source/Developer/HoloLens/HoloLensPlatformEditor/Public/HoloLensPlatformEditor.h
deleted file mode 100644
diff --git a/Engine/Source/Developer/HoloLens/HoloLensTargetPlatform/HoloLensTargetPlatform.Build.cs b/Engine/Source/Developer/HoloLens/HoloLensTargetPlatform/HoloLensTargetPlatform.Build.cs
--- a/Engine/Source/Developer/HoloLens/HoloLensTargetPlatform/HoloLensTargetPlatform.Build.cs
+++ b/Engine/Source/Developer/HoloLens/HoloLensTargetPlatform/HoloLensTargetPlatform.Build.cs
@@ -17,8 +17,7 @@ public class HoloLensTargetPlatform : ModuleRules
                 "TargetPlatform",
 				"DesktopPlatform",
 				"HTTP",
-                "Json",
-                "HoloLensPlatformEditor",
+                "Json"
             }
         );
 
diff --git a/Engine/Source/Developer/HoloLens/HoloLensTargetPlatform/Private/HoloLensTargetPlatform.cpp b/Engine/Source/Developer/HoloLens/HoloLensTargetPlatform/Private/HoloLensTargetPlatform.cpp
--- a/Engine/Source/Developer/HoloLens/HoloLensTargetPlatform/Private/HoloLensTargetPlatform.cpp
+++ b/Engine/Source/Developer/HoloLens/HoloLensTargetPlatform/Private/HoloLensTargetPlatform.cpp
@@ -5,17 +5,80 @@
 =============================================================================*/
 
 #include "HoloLensTargetPlatform.h"
+#include "HoloLensTargetTypes.h"
 #include "HoloLensTargetDevice.h"
 #include "Misc/ConfigCacheIni.h"
 #include "Misc/ScopeLock.h"
 #include "HttpModule.h"
 #include "PlatformHttp.h"
 #include "Interfaces/IHttpResponse.h"
-#include "HoloLensPlatformEditor.h"
 #include "GeneralProjectSettings.h"
 
 DEFINE_LOG_CATEGORY(LogHoloLensTargetPlatform);
 
+TArray<FHoloLensSDKVersion> FHoloLensSDKVersion::GetSDKVersions()
+{
+	bool bFound18362 = false;
+
+	static TArray<FHoloLensSDKVersion> SDKVersions;
+	if (SDKVersions.Num() == 0)
+	{
+		HKEY hRoots[] = { HKEY_LOCAL_MACHINE , HKEY_CURRENT_USER };
+
+		for (DWORD dwRootIdx = 0; dwRootIdx < _countof(hRoots); ++dwRootIdx)
+		{
+			HKEY hKey;
+			if (RegOpenKeyExW(hRoots[dwRootIdx], L"Software\\Microsoft\\Windows Kits\\Installed Roots", 0, KEY_ENUMERATE_SUB_KEYS | KEY_QUERY_VALUE, &hKey) == ERROR_SUCCESS)
+			{
+				DWORD dwIndex = 0;
+				WCHAR sSubKeyName[256];
+				DWORD dwSubKeyNameLen = _countof(sSubKeyName);
+				FString SDKPath;
+
+				if (ERROR_SUCCESS == RegQueryValueExW(hKey, L"KitsRoot10", 0, NULL, (LPBYTE)sSubKeyName, &dwSubKeyNameLen))
+				{
+					SDKPath = FString(dwSubKeyNameLen, sSubKeyName);
+				}
+				dwSubKeyNameLen = _countof(sSubKeyName);
+
+				while (RegEnumKeyExW(hKey, dwIndex, sSubKeyName, &dwSubKeyNameLen, NULL, NULL, NULL, NULL) == ERROR_SUCCESS)
+				{
+					int version[4];
+					if (swscanf_s(sSubKeyName, L"%d.%d.%d.%d", &version[0], &version[1], &version[2], &version[3]) == 4)
+					{
+						if (FPaths::FileExists(FPaths::Combine(SDKPath, TEXT("Include"), FString(dwSubKeyNameLen, sSubKeyName), TEXT("um"), TEXT("Windows.h"))))
+						{
+							// HoloLens min version is 17763 for remoting and 18317 for hand tracking
+							// @todo JoeG - update this version when the official os/sdk are released
+							if (version[0] >= 10 && version[2] >= 17763)
+							{
+								new(SDKVersions) FHoloLensSDKVersion(FString(dwSubKeyNameLen, sSubKeyName), version[0], version[1], version[2], version[3]);
+							}
+							if (version[2] == 18362)
+							{
+								bFound18362 = true;
+							}
+						}
+					}
+					dwSubKeyNameLen = _countof(sSubKeyName);
+					dwIndex++;
+				}
+				RegCloseKey(hKey);
+			}
+		}
+
+		// 18362 is required as the official minimum SDK released by Microsoft for all needed HL2 features to work
+		// If it's not found, we still need it as a hidden SDK version so that we can set MinimumPlatformVersion correctly to
+		// allow the package to run on lower installed firmware revisions
+		if (!bFound18362)
+		{
+			new(SDKVersions) FHoloLensSDKVersion(FString(L"10.0.18362.0"), 10, 0, 18362, 0);
+		}
+	}
+
+	return SDKVersions;
+}
+
 FHoloLensTargetPlatform::FHoloLensTargetPlatform(bool bIsClientOnly)
 	: TNonDesktopTargetPlatformBase(bIsClientOnly)
 {
diff --git a/Engine/Source/Developer/HoloLens/HoloLensTargetPlatform/Public/HoloLensTargetTypes.h b/Engine/Source/Developer/HoloLens/HoloLensTargetPlatform/Public/HoloLensTargetTypes.h
new file mode 100644
--- /dev/null
+++ b/Engine/Source/Developer/HoloLens/HoloLensTargetPlatform/Public/HoloLensTargetTypes.h
@@ -0,0 +1,31 @@
+// Copyright Epic Games, Inc. All Rights Reserved.
+
+#include "CoreMinimal.h"
+
+struct HOLOLENSTARGETPLATFORM_API FHoloLensSDKVersion
+{
+	/** The whole version string for the SDK */
+	FString VersionString;
+	/** The parsed parts of the version string (1.2.3.4) */
+	int32 Version1;
+	int32 Version2;
+	int32 Version3;
+	int32 Version4;
+
+	FHoloLensSDKVersion()
+	{
+		Version1 = Version2 = Version3 = Version4 = 0;
+	}
+
+	FHoloLensSDKVersion(const FString& InVersionString, int32 InVer1, int32 InVer2, int32 InVer3, int32 InVer4)
+		: VersionString(InVersionString)
+		, Version1(InVer1)
+		, Version2(InVer2)
+		, Version3(InVer3)
+		, Version4(InVer4)
+	{
+	}
+
+	/** Returns the list of SDK versions for this machine */
+	static TArray<FHoloLensSDKVersion> GetSDKVersions();
+};
\ No newline at end of file
