Ignore Xbox devices during Windows Device Portal enumeration for Hololens.
 - there's libcurl warnings because they always require https anyway.

#rnx

diff --git a/Engine/Source/Developer/HoloLens/HoloLensTargetPlatform/Private/HoloLensDeviceDetector.cpp b/Engine/Source/Developer/HoloLens/HoloLensTargetPlatform/Private/HoloLensDeviceDetector.cpp
--- a/Engine/Source/Developer/HoloLens/HoloLensTargetPlatform/Private/HoloLensDeviceDetector.cpp
+++ b/Engine/Source/Developer/HoloLens/HoloLensTargetPlatform/Private/HoloLensDeviceDetector.cpp
@@ -354,6 +354,13 @@ void FHoloLensDeviceDetector::OnDeviceWatcherDeviceAdded(IDeviceInformation* Inf
 		DeviceIp = hstr.GetRawBuffer(nullptr);
 	}
 
+	// Xbox always requires authentication & https, so no point in even trying
+	if (NewDevice.DeviceTypeName == TEXT("Windows.Xbox"))
+	{
+		UE_LOG(LogHoloLensTargetPlatform, Verbose, TEXT("ignoring %s %s"), *NewDevice.DeviceTypeName.ToString(), *DeviceIp);
+		return;
+	}
+
 	NewDevice.WdpUrl = FString::Printf(TEXT("%s://%s:%d/"), *Protocol, *DeviceIp, WdpPort);
 
 
