UE-142973 Selecting "Browse" from the project dropdown menu causes UnrealFrontend.exe to hang
-Swapped ::RoInitialize to FWindowsPlatformMisc::CoInitialize().  Not sure this is the final fix, and it may further break hololens device detection, but we can't be setting the main thread up differently for this platform.

diff --git a/Engine/Source/Developer/HoloLens/HoloLensTargetPlatform/Private/HoloLensDeviceDetector.cpp b/Engine/Source/Developer/HoloLens/HoloLensTargetPlatform/Private/HoloLensDeviceDetector.cpp
--- a/Engine/Source/Developer/HoloLens/HoloLensTargetPlatform/Private/HoloLensDeviceDetector.cpp
+++ b/Engine/Source/Developer/HoloLens/HoloLensTargetPlatform/Private/HoloLensDeviceDetector.cpp
@@ -161,11 +161,15 @@ void FHoloLensDeviceDetector::StartDeviceDetection()
 
 #if APPXPACKAGING_ENABLE
 
-	HRESULT hr;
-
-	hr = ::RoInitialize(RO_INIT_MULTITHREADED);
-	if (FAILED(hr)) { return; }
+	// We cannot use RoInitialize(multithreaded) because this is called from the main thread and other systems need single threaded apartment, and use CoInitialize to set that up.
+	// If we do need to use RoInitialize we would need to move the winrt code here to a separate thread.
+	bool bSuccess = FWindowsPlatformMisc::CoInitialize();
+	if (!bSuccess) { return; }
+	//HRESULT hr1;
+	//hr1 = ::RoInitialize(RO_INIT_MULTITHREADED);
+	//if (FAILED(hr1)) { return; }
 
+	HRESULT hr;
 	ComPtr <IDeviceInformationStatics2> DeviceInformationStatics2;
 	hr = GetActivationFactory(HStringReference(RuntimeClass_Windows_Devices_Enumeration_DeviceInformation).Get(), &DeviceInformationStatics2);
 	if (FAILED(hr)) { return; }
