UE-144739 HoloLens device detection ::RoInitialize(RO_INIT_MULTITHREADED) problem.

diff --git a/Engine/Source/Developer/HoloLens/HoloLensTargetPlatform/Private/HoloLensDeviceDetector.cpp b/Engine/Source/Developer/HoloLens/HoloLensTargetPlatform/Private/HoloLensDeviceDetector.cpp
--- a/Engine/Source/Developer/HoloLens/HoloLensTargetPlatform/Private/HoloLensDeviceDetector.cpp
+++ b/Engine/Source/Developer/HoloLens/HoloLensTargetPlatform/Private/HoloLensDeviceDetector.cpp
@@ -164,7 +164,10 @@ void FHoloLensDeviceDetector::StartDeviceDetection()
 	// We cannot use RoInitialize(multithreaded) because this is called from the main thread and other systems need single threaded apartment, and use CoInitialize to set that up.
 	// If we do need to use RoInitialize we would need to move the winrt code here to a separate thread.
 	bool bSuccess = FWindowsPlatformMisc::CoInitialize();
-	if (!bSuccess) { return; }
+	if (!bSuccess) 
+	{ 
+		return; 
+	}
 	//HRESULT hr1;
 	//hr1 = ::RoInitialize(RO_INIT_MULTITHREADED);
 	//if (FAILED(hr1)) { return; }
@@ -236,8 +239,8 @@ void FHoloLensDeviceDetector::StopDeviceDetection()
 void FHoloLensDeviceDetector::OnDeviceWatcherDeviceAdded(IDeviceInformation* Info)
 {
 	SslCertDisabler disabler;
-	FHoloLensDeviceInfo NewDevice;
-	uint32 WdpPort;
+	FHoloLensDeviceInfo NewDevice = {};
+	uint32 WdpPort = 0;
 	FString DeviceIp;
 	FString Protocol = TEXT("http");
 	FName DeviceType = NAME_None;
@@ -382,7 +385,21 @@ void FHoloLensDeviceDetector::OnDeviceWatcherDeviceAdded(IDeviceInformation* Inf
 		{
 			int32 ResponseCode = HttpResponse->GetResponseCode();
 			RequiresCredentials = ResponseCode == EHttpResponseCodes::Denied;
-			ReachedDevicePortal = RequiresCredentials || EHttpResponseCodes::IsOk(ResponseCode);
+			// Do not add the device if we failed on credentials because there is currently no way to recover from creating a device with the wrong credentials.
+			// By failing here we make it possible to manually connect with the correct user/password later.
+			ReachedDevicePortal = EHttpResponseCodes::IsOk(ResponseCode);
+			if (RequiresCredentials)
+			{
+				UE_LOG(LogHoloLensTargetPlatform, Warning, TEXT("HoloLens device automatic connect attempt denied (http response code %i). Perhaps the user or password is wrong? Use 'add unlisted device' to connect."), ResponseCode);
+			}
+			else if (!ReachedDevicePortal)
+			{
+				UE_LOG(LogHoloLensTargetPlatform, Warning, TEXT("HoloLens device automatic connect attempt failed with http response code %i. You may be able to use 'add unlisted device' to connect"), ResponseCode);
+			}
+		}
+		else
+		{
+			UE_LOG(LogHoloLensTargetPlatform, Warning, TEXT("HoloLens device automatic connect attempt failed without http response. Previous log lines may contain more information about the failed request (often bad ip/port). You may be able to use 'add unlisted device' to connect"));
 		}
 
 		if (ReachedDevicePortal || NewDevice.IsLocal)
@@ -482,13 +499,20 @@ void FHoloLensDeviceDetector::TryAddDevice(const FString& DeviceId, const FStrin
 	{
 		if (!bSucceeded || !HttpResponse.IsValid())
 		{
-			UE_LOG(LogHoloLensTargetPlatform, Error, TEXT("The device %s doesn't respond, the adding fails"), *DeviceInfo->WdpUrl);
+			UE_LOG(LogHoloLensTargetPlatform, Error, TEXT("The device %s didn't respond, the add fails.  Previous log lines may contain more information.  (often ip/port incorrect)"), *DeviceInfo->WdpUrl);
 			return;
 		}
 
-		if (HttpResponse->GetResponseCode() != 200) //we can accept just 200 OK here
+		if (HttpResponse->GetResponseCode() != EHttpResponseCodes::Ok) //we can accept just 200 OK here
 		{
-			UE_LOG(LogHoloLensTargetPlatform, Error, TEXT("The device %s responds with the http error %d, the adding fails"), *DeviceInfo->WdpUrl, HttpResponse->GetResponseCode());
+			if (HttpResponse->GetResponseCode() == EHttpResponseCodes::Denied)
+			{
+				UE_LOG(LogHoloLensTargetPlatform, Error, TEXT("The device %s responded with the http error %d, the add fails.  (this code suggests user or password are incorrect."), *DeviceInfo->WdpUrl, HttpResponse->GetResponseCode());
+			}
+			else
+			{
+				UE_LOG(LogHoloLensTargetPlatform, Error, TEXT("The device %s responded with the http error %d, the add fails."), *DeviceInfo->WdpUrl, HttpResponse->GetResponseCode());
+			}
 			return;
 		}
 
diff --git a/Engine/Source/Developer/HoloLens/HoloLensTargetPlatform/Private/IHoloLensDeviceDetector.h b/Engine/Source/Developer/HoloLens/HoloLensTargetPlatform/Private/IHoloLensDeviceDetector.h
--- a/Engine/Source/Developer/HoloLens/HoloLensTargetPlatform/Private/IHoloLensDeviceDetector.h
+++ b/Engine/Source/Developer/HoloLens/HoloLensTargetPlatform/Private/IHoloLensDeviceDetector.h
@@ -29,8 +29,8 @@ struct FHoloLensDeviceInfo
 	EHoloLensArchitecture Architecture;
 	FString Username;
 	FString Password;
-	int32 RequiresCredentials : 1;
-	int32 IsLocal : 1;
+	uint8 RequiresCredentials : 1;
+	uint8 IsLocal : 1;
 };
 
 typedef TSharedPtr<class IHoloLensDeviceDetector, ESPMode::ThreadSafe> IHoloLensDeviceDetectorPtr;
