UE-147010 Hololens crash when opening this test level in ShooterGame as a server
-No longer adding hololens target platform when we are running as a dedicated server (eg editor with -server).  Servers have no reason to detect hololens devices.

diff --git a/Engine/Source/Developer/HoloLens/HoloLensTargetPlatform/Private/HoloLensTargetPlatform.cpp b/Engine/Source/Developer/HoloLens/HoloLensTargetPlatform/Private/HoloLensTargetPlatform.cpp
--- a/Engine/Source/Developer/HoloLens/HoloLensTargetPlatform/Private/HoloLensTargetPlatform.cpp
+++ b/Engine/Source/Developer/HoloLens/HoloLensTargetPlatform/Private/HoloLensTargetPlatform.cpp
@@ -91,6 +91,10 @@ FHoloLensTargetPlatform::FHoloLensTargetPlatform(bool bIsClientOnly)
 
 	DeviceDetectedRegistration = DeviceDetector->OnDeviceDetected().AddRaw(this, &FHoloLensTargetPlatform::OnDeviceDetected);
 
+#if WITH_EDITOR
+	// Don't automatically try to detect a device if running editor with -server.
+	if (!IsRunningDedicatedServer() && !IsRunningCommandlet())
+#endif
 	{
 		DeviceDetector->StartDeviceDetection();
 	}
