move HoloLens default packaging resource images to platform extension folder.

#rnx

diff --git a/Engine/Source/Developer/HoloLens/HoloLensPlatformEditor/Private/HoloLensImageResourcesCustomization.cpp b/Engine/Source/Developer/HoloLens/HoloLensPlatformEditor/Private/HoloLensImageResourcesCustomization.cpp
--- a/Engine/Source/Developer/HoloLens/HoloLensPlatformEditor/Private/HoloLensImageResourcesCustomization.cpp
+++ b/Engine/Source/Developer/HoloLens/HoloLensPlatformEditor/Private/HoloLensImageResourcesCustomization.cpp
@@ -37,7 +37,7 @@ namespace
 
 	void AddWidgetForResourceImage(IDetailChildrenBuilder& ChildrenBuilder, const FString& RootPath, const FString& Culture, const FString& ImageFileName, const FText& ImageCaption, const FVector2D& ImageDimensions)
 	{
-		const FString DefaultEngineImageSubPath = FString::Printf(TEXT("Build/HoloLens/DefaultImages/%s.png"), *ImageFileName);
+		const FString DefaultEngineImageSubPath = FString::Printf(TEXT("Platforms/HoloLens/Build/DefaultImages/%s.png"), *ImageFileName);
 		const FString DefaultGameImageSubPath = FString::Printf(TEXT("Build/HoloLens/Resources/%s/%s.png"), *Culture, *ImageFileName);
 
 		const FString EngineImagePath = FPaths::EngineDir() / DefaultEngineImageSubPath;
